Moving some code into namespaces to prevent collision in builds and fix CIS.

diff --git a/Engine/Plugins/Tests/EditorTests/Source/EditorTests/Private/UnrealEd/AutoReimportTests.cpp b/Engine/Plugins/Tests/EditorTests/Source/EditorTests/Private/UnrealEd/AutoReimportTests.cpp
--- a/Engine/Plugins/Tests/EditorTests/Source/EditorTests/Private/UnrealEd/AutoReimportTests.cpp
+++ b/Engine/Plugins/Tests/EditorTests/Source/EditorTests/Private/UnrealEd/AutoReimportTests.cpp
@@ -102,9 +102,12 @@ struct FAutoReimportTestPayload
 	}
 };
 
-FString GetWorkingDir()
+namespace AutoReimportTests
 {
-	return FPaths::ConvertRelativePathToFull(FPaths::AutomationTransientDir()) / TEXT("AutoReimport") / TEXT("");
+	FString GetWorkingDir()
+	{
+		return FPaths::ConvertRelativePathToFull(FPaths::AutomationTransientDir()) / TEXT("AutoReimport") / TEXT("");
+	}
 }
 
 FString GetTestSourceDir()
@@ -143,7 +146,7 @@ bool CopyTestFiles(const FAutoReimportTestPayload& Test, const TArray<FSrcDstFil
 IMPLEMENT_SIMPLE_AUTOMATION_TEST(FAutoReimportSimpleCreateTest, "Editor.Auto Reimport.Simple Create", EAutomationTestFlags::EditorContext | EAutomationTestFlags::EngineFilter)
 bool FAutoReimportSimpleCreateTest::RunTest(const FString& Parameters)
 {
-	const FString WorkingDir = GetWorkingDir();
+	const FString WorkingDir = AutoReimportTests::GetWorkingDir();
 
 	static const TCHAR* Filename = TEXT("square.png");
 
@@ -199,7 +202,7 @@ bool FAutoReimportSimpleCreateTest::RunTest(const FString& Parameters)
 IMPLEMENT_SIMPLE_AUTOMATION_TEST(FAutoReimportSimpleModifyTest, "Editor.Auto Reimport.Simple Modify", EAutomationTestFlags::EditorContext | EAutomationTestFlags::EngineFilter)
 bool FAutoReimportSimpleModifyTest::RunTest(const FString& Parameters)
 {
-	const FString WorkingDir = GetWorkingDir();
+	const FString WorkingDir = AutoReimportTests::GetWorkingDir();
 
 	static const TCHAR* SrcFilename1 = TEXT("square.png");
 	static const TCHAR* SrcFilename2 = TEXT("red-square.png");
@@ -271,7 +274,7 @@ bool FAutoReimportSimpleModifyTest::RunTest(const FString& Parameters)
 IMPLEMENT_SIMPLE_AUTOMATION_TEST(FAutoReimportSimpleDeleteTest, "Editor.Auto Reimport.Simple Delete", EAutomationTestFlags::EditorContext | EAutomationTestFlags::EngineFilter)
 bool FAutoReimportSimpleDeleteTest::RunTest(const FString& Parameters)
 {
-	const FString WorkingDir = GetWorkingDir();
+	const FString WorkingDir = AutoReimportTests::GetWorkingDir();
 
 	static const TCHAR* Filename = TEXT("square.png");
 
@@ -336,7 +339,7 @@ bool FAutoReimportSimpleDeleteTest::RunTest(const FString& Parameters)
 IMPLEMENT_SIMPLE_AUTOMATION_TEST(FAutoReimportSimpleRenameTest, "Editor.Auto Reimport.Simple Rename", EAutomationTestFlags::EditorContext | EAutomationTestFlags::EngineFilter)
 bool FAutoReimportSimpleRenameTest::RunTest(const FString& Parameters)
 {
-	const FString WorkingDir = GetWorkingDir();
+	const FString WorkingDir = AutoReimportTests::GetWorkingDir();
 
 	static const TCHAR* SrcFilename = TEXT("square.png");
 	static const TCHAR* DstFilename = TEXT("square2.png");
@@ -401,7 +404,7 @@ bool FAutoReimportSimpleRenameTest::RunTest(const FString& Parameters)
 IMPLEMENT_SIMPLE_AUTOMATION_TEST(FAutoReimportSimpleMoveExternallyTest, "Editor.Auto Reimport.Move Externally", EAutomationTestFlags::EditorContext | EAutomationTestFlags::EngineFilter)
 bool FAutoReimportSimpleMoveExternallyTest::RunTest(const FString& Parameters)
 {
-	const FString WorkingDir = GetWorkingDir();
+	const FString WorkingDir = AutoReimportTests::GetWorkingDir();
 
 	static const TCHAR* SrcFilename = TEXT("square.png");
 	static const TCHAR* DstFilename = TEXT("../square.png");
@@ -467,7 +470,7 @@ bool FAutoReimportSimpleMoveExternallyTest::RunTest(const FString& Parameters)
 IMPLEMENT_SIMPLE_AUTOMATION_TEST(FAutoReimportRestartDetectionTest, "Editor.Auto Reimport.Restart Detection", EAutomationTestFlags::EditorContext | EAutomationTestFlags::EngineFilter)
 bool FAutoReimportRestartDetectionTest::RunTest(const FString& Parameters)
 {
-	const FString WorkingDir = GetWorkingDir();
+	const FString WorkingDir = AutoReimportTests::GetWorkingDir();
 	static const TCHAR* SrcFilename = TEXT("square.png");
 	static const TCHAR* SrcFilename2 = TEXT("red-square.png");
 	
@@ -566,7 +569,7 @@ bool FAutoReimportRestartDetectionTest::RunTest(const FString& Parameters)
 IMPLEMENT_SIMPLE_AUTOMATION_TEST(FAutoReimportMultipleChangesTest, "Editor.Auto Reimport.Multiple Changes", EAutomationTestFlags::EditorContext | EAutomationTestFlags::EngineFilter)
 bool FAutoReimportMultipleChangesTest::RunTest(const FString& Parameters)
 {
-	const FString WorkingDir = GetWorkingDir();
+	const FString WorkingDir = AutoReimportTests::GetWorkingDir();
 	static const TCHAR* Filename1 = TEXT("square.png");
 	static const TCHAR* Filename2 = TEXT("red-square.png");
 	
@@ -622,7 +625,7 @@ bool FAutoReimportMultipleChangesTest::RunTest(const FString& Parameters)
 IMPLEMENT_SIMPLE_AUTOMATION_TEST(FAutoReimportChangeExtensionsTest, "Editor.Auto Reimport.Change Extensions", EAutomationTestFlags::EditorContext | EAutomationTestFlags::EngineFilter)
 bool FAutoReimportChangeExtensionsTest::RunTest(const FString& Parameters)
 {
-	const FString WorkingDir = GetWorkingDir();
+	const FString WorkingDir = AutoReimportTests::GetWorkingDir();
 	static const TCHAR* SrcFilename1 = TEXT("square.png");
 	static const TCHAR* SrcFilename2 = TEXT("red-square.png");
 	static const TCHAR* SrcFilename3 = TEXT("empty.txt");
@@ -714,7 +717,7 @@ bool FAutoReimportChangeExtensionsTest::RunTest(const FString& Parameters)
 IMPLEMENT_SIMPLE_AUTOMATION_TEST(FAutoReimportWildcardFiltersTest, "Editor.Auto Reimport.Wildcard Filters", EAutomationTestFlags::EditorContext | EAutomationTestFlags::EngineFilter)
 bool FAutoReimportWildcardFiltersTest::RunTest(const FString& Parameters)
 {
-	const FString WorkingDir = GetWorkingDir();
+	const FString WorkingDir = AutoReimportTests::GetWorkingDir();
 	static const TCHAR* SrcFilename1 = TEXT("square.png");
 	static const TCHAR* SrcFilename2 = TEXT("red-square.png");
 	static const TCHAR* SrcFilename3 = TEXT("empty.txt");
@@ -729,7 +732,7 @@ bool FAutoReimportWildcardFiltersTest::RunTest(const FString& Parameters)
 	Test->Config.Rules.AddWildcardRule(TEXT("sub-folder/*.png"), false);
 	Test->Config.bDetectChangesSinceLastRun = true;
 
-	IFileManager::Get().MakeDirectory(*(GetWorkingDir() / TEXT("Content") / TEXT("sub-folder")), false);
+	IFileManager::Get().MakeDirectory(*( AutoReimportTests::GetWorkingDir() / TEXT("Content") / TEXT("sub-folder")), false);
 		
 	{
 		TArray<FSrcDstFilenames> Files;
@@ -760,7 +763,7 @@ bool FAutoReimportWildcardFiltersTest::RunTest(const FString& Parameters)
 		}
 		else
 		{
-			// Check that we get the same behaviour for live changes
+			// Check that we get the same behavior for live changes
 			TArray<FSrcDstFilenames> Files;
 			Files.Emplace(SrcFilename2, DstFilename1);
 			Files.Emplace(SrcFilename1, DstFilename2);
diff --git a/Engine/Source/Developer/DirectoryWatcher/Private/DirectoryWatcherTests.cpp b/Engine/Source/Developer/DirectoryWatcher/Private/DirectoryWatcherTests.cpp
--- a/Engine/Source/Developer/DirectoryWatcher/Private/DirectoryWatcherTests.cpp
+++ b/Engine/Source/Developer/DirectoryWatcher/Private/DirectoryWatcherTests.cpp
@@ -74,16 +74,19 @@ struct FDirectoryWatcherTestPayload
 	}
 };
 
-FString GetWorkingDir()
+namespace DirectoryWatcherTests
 {
-	return FPaths::ConvertRelativePathToFull(FPaths::AutomationTransientDir() / TEXT("DirectoryWatcher")) / TEXT("");
+	FString GetWorkingDir()
+	{
+		return FPaths::ConvertRelativePathToFull(FPaths::AutomationTransientDir() / TEXT("DirectoryWatcher")) / TEXT("");
+	}
 }
 
 static const float TestTickDelay = 1.0f;
 IMPLEMENT_SIMPLE_AUTOMATION_TEST(FDirectoryWatcherSimpleCreateTest, "System.Plugins.Directory Watcher.Simple Create", EAutomationTestFlags::EditorContext | EAutomationTestFlags::EngineFilter)
 bool FDirectoryWatcherSimpleCreateTest::RunTest(const FString& Parameters)
 {
-	const FString WorkingDir = GetWorkingDir();
+	const FString WorkingDir = DirectoryWatcherTests::GetWorkingDir();
 
 	static const TCHAR* Filename = TEXT("created.tmp");
 
@@ -115,7 +118,7 @@ bool FDirectoryWatcherSimpleCreateTest::RunTest(const FString& Parameters)
 IMPLEMENT_SIMPLE_AUTOMATION_TEST(FDirectoryWatcherSimpleModifyTest, "System.Plugins.Directory Watcher.Simple Modify", EAutomationTestFlags::EditorContext | EAutomationTestFlags::EngineFilter)
 bool FDirectoryWatcherSimpleModifyTest::RunTest(const FString& Parameters)
 {
-	const FString WorkingDir = GetWorkingDir();
+	const FString WorkingDir = DirectoryWatcherTests::GetWorkingDir();
 
 	static const TCHAR* Filename = TEXT("modified.tmp");
 
@@ -153,7 +156,7 @@ bool FDirectoryWatcherSimpleModifyTest::RunTest(const FString& Parameters)
 IMPLEMENT_SIMPLE_AUTOMATION_TEST(FDirectoryWatcherSimpleDeleteTest, "System.Plugins.Directory Watcher.Simple Delete", EAutomationTestFlags::EditorContext | EAutomationTestFlags::EngineFilter)
 bool FDirectoryWatcherSimpleDeleteTest::RunTest(const FString& Parameters)
 {
-	const FString WorkingDir = GetWorkingDir();
+	const FString WorkingDir = DirectoryWatcherTests::GetWorkingDir();
 
 	static const TCHAR* Filename = TEXT("removed.tmp");
 
@@ -188,7 +191,7 @@ bool FDirectoryWatcherSimpleDeleteTest::RunTest(const FString& Parameters)
 IMPLEMENT_SIMPLE_AUTOMATION_TEST(FDirectoryWatcherSubFolderTest, "System.Plugins.Directory Watcher.Sub Folder", EAutomationTestFlags::EditorContext | EAutomationTestFlags::EngineFilter)
 bool FDirectoryWatcherSubFolderTest::RunTest(const FString& Parameters)
 {
-	const FString WorkingDir = GetWorkingDir();
+	const FString WorkingDir = DirectoryWatcherTests::GetWorkingDir();
 
 	static const TCHAR* CreatedFilename = TEXT("sub_folder/created.tmp");
 	static const TCHAR* ModifiedFilename = TEXT("sub_folder/modified.tmp");
@@ -247,7 +250,7 @@ bool FDirectoryWatcherSubFolderTest::RunTest(const FString& Parameters)
 IMPLEMENT_SIMPLE_AUTOMATION_TEST(FDirectoryWatcherNewFolderTest, "System.Plugins.Directory Watcher.New Folder", EAutomationTestFlags::EditorContext | EAutomationTestFlags::EngineFilter)
 bool FDirectoryWatcherNewFolderTest::RunTest(const FString& Parameters)
 {
-	const FString WorkingDir = GetWorkingDir();
+	const FString WorkingDir = DirectoryWatcherTests::GetWorkingDir();
 
 	static const TCHAR* CreatedDirectory = TEXT("created");
 	static const TCHAR* RemovedDirectory = TEXT("removed");
@@ -294,7 +297,7 @@ bool FDirectoryWatcherNewFolderTest::RunTest(const FString& Parameters)
 IMPLEMENT_SIMPLE_AUTOMATION_TEST(FDirectoryWatcherIgnoreSubtreeTest, "System.Plugins.Directory Watcher.Ignore Subtree", EAutomationTestFlags::EditorContext | EAutomationTestFlags::EngineFilter)
 bool FDirectoryWatcherIgnoreSubtreeTest::RunTest(const FString& Parameters)
 {
-	const FString WorkingDir = GetWorkingDir();
+	const FString WorkingDir = DirectoryWatcherTests::GetWorkingDir();
 
 	static const TCHAR* ChildDirectory = TEXT("child");
 	static const TCHAR* GrandchildDirectory = TEXT("grandchild");
