Copying //UE4/Dev-Editor to //UE4/Dev-Main (Source: //UE4/Dev-Editor @ 3279756)

==========================
MAJOR FEATURES + CHANGES
==========================

Change 3229490 on 2016/12/09 by [REDACTED_PERSON]

	Integrated fix to support named changelists in SVN

Change 3229574 on 2016/12/09 by [REDACTED_PERSON]

	Fix actor mobility getting changed on scene reimport

Change 3229692 on 2016/12/09 by [REDACTED_PERSON]

	Fixing an XML Parser assert when parsing a root tag that completes on the same line.

Change 3230582 on 2016/12/12 by [REDACTED_PERSON]

	PR #3024: Correct the outdated error message instructions for how to fix being unable to launch on an iOS device. (Contributed by [REDACTED_PERSON])

Change 3231470 on 2016/12/12 by [REDACTED_PERSON]

	Eliminate editor sounds that play when  you PIE, simulate or possess the player.  They get in the way of game sounds, are annoying to hear when you are constantly starting and stopping pie, and flush async loading that the game might be doing when they load.

Change 3231475 on 2016/12/12 by [REDACTED_PERSON]

Change 3231476 on 2016/12/12 by [REDACTED_PERSON]

Change 3231477 on 2016/12/12 by [REDACTED_PERSON]

Change 3231479 on 2016/12/12 by [REDACTED_PERSON]

Change 3231480 on 2016/12/12 by [REDACTED_PERSON]

Change 3231508 on 2016/12/12 by [REDACTED_PERSON]

	Removing todo comment

Change 3231603 on 2016/12/12 by [REDACTED_PERSON]

	Exposed a 0-1 UV set and the scaled pixel size for Box and Border brushes
	Also added a material function that exposes all of the current UV sets with nice names instead of indexed coordinates

Change 3231618 on 2016/12/12 by [REDACTED_PERSON]

Change 3231638 on 2016/12/12 by [REDACTED_PERSON]

	Add RF_Transactional to the list of default flags for creating or importing new assets.  All should be transactional by default

Change 3231642 on 2016/12/12 by [REDACTED_PERSON]

	Brighten up the output log by default

Change 3231648 on 2016/12/12 by [REDACTED_PERSON]

Change 3231666 on 2016/12/12 by [REDACTED_PERSON]

Change 3231674 on 2016/12/12 by [REDACTED_PERSON]

Change 3231745 on 2016/12/12 by [REDACTED_PERSON]

	Back out changelist 3231477 to fix build error C2259

Change 3232417 on 2016/12/13 by [REDACTED_PERSON]

	Add the following attributes to the Editor.Usage.FBX.Import EngineAnalytics event
	- FBX Version
	- Filename Hash
	- Import Type

Change 3232477 on 2016/12/13 by [REDACTED_PERSON]

Change 3232571 on 2016/12/13 by [REDACTED_PERSON]

	Back out changelist 3231745

Change 3232675 on 2016/12/13 by [REDACTED_PERSON]

	Fix a crash when reordering material with a fbx containing unused materials, add a fbx automation test to prevent similar issue.

Change 3232975 on 2016/12/13 by [REDACTED_PERSON]

	Fix to build error C2259 for the IPluginWizardDefinition API change.

Change 3233146 on 2016/12/13 by [REDACTED_PERSON]

	Fixed a rounding errors resulting in not flattening to the specified height
	Fixed a rounding error resulting in LandscapeDataAccess::GetTexHeight not always returning the appropriate value

Change 3233153 on 2016/12/13 by [REDACTED_PERSON]

	We cannot anymore change the instance override materials array topology, the topology is limited by the mesh materials array

Change 3234406 on 2016/12/14 by [REDACTED_PERSON]

	Fix window handle and device context being accessed by scene viewports after the underlying window has been destroyed by the OS.  This is an invalid state on linux and using some vr devices.

Change 3234485 on 2016/12/14 by [REDACTED_PERSON]

	tentative build fix for Mac

Change 3234495 on 2016/12/14 by [REDACTED_PERSON]

	Made a setting to control if PIE enter and exit sounds are played.  Off by default

Change 3236709 on 2016/12/15 by [REDACTED_PERSON]

	Fix camera export rotation offset

Change 3236782 on 2016/12/15 by [REDACTED_PERSON]

	Fixed EmitTermExpr failing to use the correct package ID

	FBPTerminal::Source used to be set to the pin, however when pins were moved away from being UObjects, FBPTerminal::SourcePin was added and FBPTerminal::Source is typically null.

Change 3236853 on 2016/12/15 by [REDACTED_PERSON]

	Fix the serialization of the staticmesh property FMeshSectionInfoMap

Change 3236890 on 2016/12/15 by [REDACTED_PERSON]

	Remove old define

Change 3239328 on 2016/12/18 by [REDACTED_PERSON]

	Fixed Focus Viewport action in Static Mesh Viewport. Problem was that the conversion to Orbit Camera for storing the camera position was trashing the desired position during cvamera transitions.  Orbit camera position is now only stored at the end of a transition.

Change 3239660 on 2016/12/19 by [REDACTED_PERSON]

Change 3239662 on 2016/12/19 by [REDACTED_PERSON]

Change 3239663 on 2016/12/19 by [REDACTED_PERSON]

Change 3240696 on 2016/12/20 by [REDACTED_PERSON]

	Added paddiing to columns view
	Added auto resize of column when double clicking on splitter handle in the header
	Remove right number alignment after discussion with Matt K.

Change 3240758 on 2016/12/20 by [REDACTED_PERSON]

	added missing non abstract implementation

Change 3240782 on 2016/12/20 by [REDACTED_PERSON]

	Added missing documentation for content browser column auto resizing

Change 3240817 on 2016/12/20 by [REDACTED_PERSON]

Change 3240834 on 2016/12/20 by [REDACTED_PERSON]

	tentative fix for build error

Change 3240984 on 2016/12/20 by [REDACTED_PERSON]

	Removed unnecessary functions

Change 3241174 on 2016/12/20 by [REDACTED_PERSON]

	Fix compile errors

Change 3241966 on 2016/12/21 by [REDACTED_PERSON]

	Fixed Typo and changed execution order in "ComboBoxString" Component
	[UE-38994] - GitHub 2971 : Fixed Typo and changed execution order in "ComboBoxString" Component

	PR #2971: Fixed Typo and changed execution order in "ComboBoxString" Component (Contributed by [REDACTED_PERSON])
	#github https://github.com/EpicGames/UnrealEngine/pull/2971

Change 3242126 on 2016/12/21 by [REDACTED_PERSON]

	Back out changelist 3236853
	We have to back out this change list because the change was implement in the 4.15 release branch and the EditorObjectVersion.h change is now implement in the ReleaseObjectVersion.h.

Change 3244492 on 2017/01/02 by [REDACTED_PERSON]

	Improved error message

Change 3244545 on 2017/01/02 by [REDACTED_PERSON]

	Navigation - Making it so we don't attempt to load HotReload during shutdown, we only access it if it's still loaded.

Change 3244549 on 2017/01/02 by [REDACTED_PERSON]

	Slate - Implementing custom hardware cursor loading across Windows, Mac and Linux and supports loading cursors from PAK files.  All platforms support loading PNGs through the FHardwareCursor interface.  Some platforms support additional formats, for multiresolution support, but there's a naming convention that can be used on PNGs for the same capability.  All of it is documented in the FHardwareCursor header.  The platform layer for ICursor, now has support for replacing cursor shapes as an override, and can be reset safely.

	The FHardwareCursor supports loading cursors from raw pixel buffers as well, the plan is to allow for the option to UTextures to also be used for hardware cursors.

	Now users through C++ can load and replace the hardware cursors with custom ones of their own,

	e.g. FSlateApplication::Get().RegisterCursor(EMouseCursor::Default, MakeShareable(new FHardwareCursor(FPaths::GameContentDir() / "Slate/FancyPointer", FIntPoint(0,0))));

	The next step is to expose a game friendly layer that supports caching cursors, and letting users change them out by name, without a bunch of destruction of OS resources.

Change 3244845 on 2017/01/03 by [REDACTED_PERSON]

	Fixing typo

Change 3244903 on 2017/01/03 by [REDACTED_PERSON]

	PR #3044: fix link error when FAssetData::PrintAssetData() is used in project (Contributed by [REDACTED_PERSON])

Change 3245125 on 2017/01/03 by [REDACTED_PERSON]

	Put back the dev-editor version because there was some data create before we back it out

Change 3246106 on 2017/01/04 by [REDACTED_PERSON]

	Removed broken CrashReportReciever pre-upload phase from CrashReportClient.
	[UE-40153] - CrashReportClient fails when used in legacy mode with a CrashReportReciever

Change 3246251 on 2017/01/04 by [REDACTED_PERSON]

Change 3246252 on 2017/01/04 by [REDACTED_PERSON]

	The "Map Directory To" field when setting directories to monitor for auto-reimport has also been changed to use the content browser path picker instead of relying on the user to manually enter a mount point.

Change 3247620 on 2017/01/05 by [REDACTED_PERSON]

	Automation - Removing an adjustment to the number of shots we take for high res shots.

Change 3247621 on 2017/01/05 by [REDACTED_PERSON]

	Automation - Adding a few more rendering tests to the cornell box.

Change 3247629 on 2017/01/05 by [REDACTED_PERSON]

	Automation - Improving the comparison row display for screenshots so it's obvious what each image represents.

Change 3248811 on 2017/01/05 by [REDACTED_PERSON]

	PR #3091: Removed unnecessary UPackage casts (Contributed by [REDACTED_PERSON])

Change 3248860 on 2017/01/06 by [REDACTED_PERSON]

	Made the plugin browser select the "built in" category by default instead of the 2D category.  There is no reason for a sub-category to be selected first as it makes searching for plugins globally an extra click because you have to click on the base category first

Change 3249264 on 2017/01/06 by [REDACTED_PERSON]

	Fixed automation test warnings

Change 3249481 on 2017/01/06 by [REDACTED_PERSON]

	Also fill new component added with the tool from neighbours predominance

Change 3249505 on 2017/01/06 by [REDACTED_PERSON]

	PR #3093: Include guard cleanup (Contributed by [REDACTED_PERSON])

Change 3249544 on 2017/01/06 by [REDACTED_PERSON]

Change 3250738 on 2017/01/09 by [REDACTED_PERSON]

	UMG - The WIC now checks if the Widget is enabled before it claims that it's over an interactable or keyboard focusable widget.

Change 3250865 on 2017/01/09 by [REDACTED_PERSON]

	Slate - Updating EAutoCenter and ESizingRule to use the newer enum class style enums.

Change 3250867 on 2017/01/09 by [REDACTED_PERSON]

	Slate - Adding more logging to the hardware cursor code so that it reports more information when it doesn't find an exact match when it comes to cursor size.

Change 3250936 on 2017/01/09 by [REDACTED_PERSON]

	Automation - Refactoring the screenshot comparison tool to no longer require one one generated report.  Doing screenshot comparions now generates individual reports for each failed comparison so that they can be evaluated in bits, and as changes occur as the user reviews aspects, we can remove the reports.  There is now async image loading for the comparison view so that it doesn't hitch.

Change 3250937 on 2017/01/09 by [REDACTED_PERSON]

	Automation - Adding another example to the CornellBox test.

Change 3250958 on 2017/01/09 by [REDACTED_PERSON]

	Slate - Fixing some other cases where people were referring to ESizingRule::Type.

Change 3251162 on 2017/01/09 by [REDACTED_PERSON]

	Slate - Fixing some other cases where people were referring to ESizingRule::Type.

Change 3251254 on 2017/01/09 by [REDACTED_PERSON]

	Attempt to fix static analysis warnings

Change 3251373 on 2017/01/09 by [REDACTED_PERSON]

	Core - Now writing a log warning instead of ensuring if calling LoadModule wouldn't have been safe to do here, depending on load order.

Change 3251525 on 2017/01/09 by [REDACTED_PERSON]

	Automation - Fixing a build issue in ImageComparer.

Change 3252321 on 2017/01/10 by [REDACTED_PERSON]

Change 3252354 on 2017/01/10 by [REDACTED_PERSON]

	Image Compare - Fixing a potential threading hazard in the image comparer.

Change 3252356 on 2017/01/10 by [REDACTED_PERSON]

	Automation - The screenshot metadata now captures the commit/CL that the screenshot was taken at and records it in the metadata.

Change 3252601 on 2017/01/10 by [REDACTED_PERSON]

	Fbx automation test, reload feature implementation

Change 3252761 on 2017/01/10 by [REDACTED_PERSON]

	Fixing some IWYU errors with PCH disabled

Change 3252765 on 2017/01/10 by [REDACTED_PERSON]

	Fixing some static analysis warnings

Change 3252793 on 2017/01/10 by [REDACTED_PERSON]

	Fixing FText natvis

	The text data visualizers have to be defined before the text visualizer

Change 3253987 on 2017/01/11 by [REDACTED_PERSON]

	PR #3108: Git Plugin: use asynchronous "MarkForAdd" and "CheckIn" operations for the initial commit (Contributed by [REDACTED_PERSON])

Change 3254378 on 2017/01/11 by [REDACTED_PERSON]

	Refactor scene importing to allow for plugins to make scene importers

Change 3254679 on 2017/01/11 by [REDACTED_PERSON]

	Fix calling LoadModule in perforce source control off the main thread

Change 3256472 on 2017/01/12 by [REDACTED_PERSON]

	Improved error reporting from IncludeTool

	- The error reporting was using zero-based line indices which was misleading.
	- The error reporting now includes the offending line to remove ambiguity.

Change 3256725 on 2017/01/13 by [REDACTED_PERSON]

	IncludeTool can now parse typedef in Fwd headers

Change 3256758 on 2017/01/13 by [REDACTED_PERSON]

	Added support for String Tables

	String Tables provide a way to centralize your localized text into one (or several) known locations, and then reference the entries within a string table from other assets or code in a robust way that allows for easy re-use of localized text.

	String Tables can be defined in C++ (using the LOCTABLE family of macros), loaded via CSV file, or created as an asset. They can be referenced in C++ using either the LOCTABLE macro, or the static FText::FromStringTable function. INI files can reference them using the LOCTABLE macro syntax, and FText properties in assets can reference them via the advanced settings combo.

Change 3257018 on 2017/01/13 by [REDACTED_PERSON]

	FbxAutomationTest fix the import reload operation, it was calling garbagecollect with no keep flag

Change 3257168 on 2017/01/13 by [REDACTED_PERSON]

	Removed code that was writing null into bytecode during save

Change 3257344 on 2017/01/13 by [REDACTED_PERSON]

	Backing out changelist 3256725, and excluding my header from the scan instead

Change 3257426 on 2017/01/13 by [REDACTED_PERSON]

	Slate - Adding the ability to invert alpha when drawing slate textures.  Going to be used in the future for rendering render targets for the scene which have inverted alpha.

Change 3257572 on 2017/01/13 by [REDACTED_PERSON]

	Slate - Fixing a build error.

Change 3257970 on 2017/01/14 by [REDACTED_PERSON]

	Fixing exclude path

Change 3258458 on 2017/01/16 by [REDACTED_PERSON]

	PR #3135: GameViewportClient: FOnCloseRequested is now a multicast delegate (Contributed by [REDACTED_PERSON])

Change 3258472 on 2017/01/16 by [REDACTED_PERSON]

	PR #3126: Fix to load editor style assets (Contributed by [REDACTED_PERSON])

Change 3258473 on 2017/01/16 by [REDACTED_PERSON]

	PR #3124: Fix wrong result with Image-DrawAsBox with PaperSprite. (Contributed by [REDACTED_PERSON])

Change 3258539 on 2017/01/16 by [REDACTED_PERSON]

	Slate - Pixel Snapping has been moved to the GPU for the RHI rendering policy.  Additionally, widgets with a render transform of Scale, Rotation or Sheer, and their children are no longer pixel snapped, this should reduce some of jittering seen by users when animations are applied to widgets.  NOTE: This only affects render transforms, any transform in layout space is still subject to pixel snapping.

Change 3258607 on 2017/01/16 by [REDACTED_PERSON]

	Fixing the mac build.

Change 3258661 on 2017/01/16 by [REDACTED_PERSON]

	Actors with experimental components no longer say
	"Uses experimental class: Actor" when selecting the actor root in the details panel

Change 3258678 on 2017/01/16 by [REDACTED_PERSON]

	Platform - Introducing a way to get the mimetype for a file on Windows.  Other platforms don't yet have an implementation outside of returning application/unknown.

Change 3258924 on 2017/01/16 by [REDACTED_PERSON]

	Platform - Implementing a fallback for the generic platform http, that can do some basic mimetype lookups.

Change 3258929 on 2017/01/16 by [REDACTED_PERSON]

	UMG - Fixing the animation to finish the evaluation before it notifies that the animation completed.

Change 3259109 on 2017/01/16 by [REDACTED_PERSON]

	Platform - The GetMimeType function now only takes in FilePath, since some platforms will require that actually resolve to a file on disk in order to determine the true mimetype.

Change 3259111 on 2017/01/16 by [REDACTED_PERSON]

	Avoid to move the camera when we re-import in the static mesh editor

Change 3259275 on 2017/01/16 by [REDACTED_PERSON]

	Fix crash when a slate window is resized and calls into a scene viewport during loading code when the scene viewport is not in a slate hierarchy and thus has no widget

Change 3259300 on 2017/01/16 by [REDACTED_PERSON]

	UMG - Introducing PreConstruct and NativePreConstruct to the base UUserWidget.  Users can now visualize non-binding based changes in the designer by evaluating a very limited amount of the blueprint code.  In the event your user widget crashes on load, due to calling something unsafe, you can disable evaluation in the editor preferences under Widget Designer.

Change 3259306 on 2017/01/16 by [REDACTED_PERSON]

	Games - Removing the Game Specific implementations of PreConstruct.

Change 3260182 on 2017/01/17 by [REDACTED_PERSON]

	Fix static analysis

Change 3261049 on 2017/01/17 by [REDACTED_PERSON]

	Slate - Putting in some fixes for the non-gpu pixel snapping mode, and disabling gpu snapping while we dig into why it looks weird.

Change 3261434 on 2017/01/17 by [REDACTED_PERSON]

	Fixing the mac build.

Change 3261435 on 2017/01/17 by [REDACTED_PERSON]

	Slate - Tweaking some aspects of the slate rounding code on the GPU.  There's still some precision loss somewhere causing subtle differences in where the snap occurs, that's different from previously.

Change 3261460 on 2017/01/17 by [REDACTED_PERSON]

	UMG - Tweaking the defintiions of NativePreConstruct, dropping passing in design time since that is readily available in native code.

Change 3261833 on 2017/01/18 by [REDACTED_PERSON]

	Fix all warning for fbx automation tests

Change 3261874 on 2017/01/18 by [REDACTED_PERSON]

	PR #3136: Fix Submit to Source Control Window for Git plugin : use CanCheckIn() to filter out unmodified assets files (Contributed by [REDACTED_PERSON])

Change 3262000 on 2017/01/18 by [REDACTED_PERSON]

	Updated Slate to allocate widgets using MakeShared

	This saves one allocation per-widget

Change 3262003 on 2017/01/18 by [REDACTED_PERSON]

	UMG - Widget Interaction Components now ignore Visible(false) Widget Components when tracing.

Change 3262052 on 2017/01/18 by [REDACTED_PERSON]

	Put back the staticmesh skinxx workflow

Change 3262775 on 2017/01/18 by [REDACTED_PERSON]

	Slate - Ditching moving vertex rounding to the GPU, some precision issues could not be overcome.  Ended up writing a clean way to implement it on the CPU.

Change 3262818 on 2017/01/18 by [REDACTED_PERSON]

Change 3263679 on 2017/01/19 by [REDACTED_PERSON]

	Slate - Adding some comments to the Slate Vertex Rounder.

Change 3265154 on 2017/01/19 by [REDACTED_PERSON]

	Slate/UMG - Putting in some more time into pixel snapping.  I've re-introduced the old constructors, and decided to go with the templated approach, as to not break old code that relied on the FSlateVertex working a certain way.

Change 3265478 on 2017/01/20 by [REDACTED_PERSON]

	Added config support for hang detection time and switching hang detection on/off in UnrealWatchdog
	[UE-40838] - Make hang time configurable and increase default in UnrealWatchdog

Change 3265600 on 2017/01/20 by [REDACTED_PERSON]

	Slate - Making some const local variables const.

Change 3265714 on 2017/01/20 by [REDACTED_PERSON]

Change 3265865 on 2017/01/20 by [REDACTED_PERSON]

Change 3267989 on 2017/01/23 by [REDACTED_PERSON]

	Exposed String Tables to Blueprints

Change 3268018 on 2017/01/23 by [REDACTED_PERSON]

	Small API clean-up for string tables

Change 3268455 on 2017/01/23 by [REDACTED_PERSON]

	Fix SaveAs (Which says SaveCurrentAs) not saving the current level and only saving the persistent level and then reloading everything thus causing work to be lost if editing a sub-level

Change 3269388 on 2017/01/24 by [REDACTED_PERSON]

	Refactored tick timing in UnrealWatchdog to stop bug where it doesn't close.
	[UE-40839] - UnrealWatchdog running and blocking use of Unreal Game Sync for internal users

	Standalone tool code only - doesn't touch engine

Change 3270205 on 2017/01/24 by [REDACTED_PERSON]

	Updated FUnrealEdMisc::OnMessageTokenActivated to properly traverse up the outer hierarchy of an object.

Change 3270231 on 2017/01/24 by [REDACTED_PERSON]

	Renamed and exposed GetFullScreenAlignment and GetViewportAnchors for consistency with the setters

Change 3271734 on 2017/01/25 by [REDACTED_PERSON]

	Add sorting for landscape target layer, user can now sort alphabetical, material based or custom
	Added a new vertical box SDragNDropVerticalBox to handle drag & drop of FSlot
	Fixed SDropTarget to only consider the drop action if it was started by it
	Added visibility toggle to only show used layers in the currently loaded data

Change 3271797 on 2017/01/25 by [REDACTED_PERSON]

	Renamed HasBeenAlreadyMadeSharable to DoesSharedInstanceExist as the old name was nonsense

Change 3271813 on 2017/01/25 by [REDACTED_PERSON]

	Fixed bad access of a shared this during widget destruction when a context menu was open

Change 3271988 on 2017/01/25 by [REDACTED_PERSON]

	Slate - Removing some old checkbox deprecated code from the 4.3 and 4.6 days.

Change 3271992 on 2017/01/25 by [REDACTED_PERSON]

	Blueprints - Making the checked call better to log out more information when dragging and dropping a missing property.

Change 3272134 on 2017/01/25 by [REDACTED_PERSON]

	Updated the GatherText commandlet to no longer hold a ConfigFile pointer while it runs

	This pointer is internal to GConfig, and may be updated (or invalidated) when other config files are loaded (as can happen via game code while gathering text).

Change 3272301 on 2017/01/25 by [REDACTED_PERSON]

	Slate - More cleanup from the removal of a old legacy enum that people were still using.

Change 3273070 on 2017/01/26 by [REDACTED_PERSON]

	Fix CIS errors in landscape code from CL 3271734

Change 3273123 on 2017/01/26 by [REDACTED_PERSON]

	Fix crash during init of CRC when running packaged without access to main engine config hierarchy.

Change 3273194 on 2017/01/26 by [REDACTED_PERSON]

	Fixing some build warnings.

Change 3273242 on 2017/01/26 by [REDACTED_PERSON]

Change 3273279 on 2017/01/26 by [REDACTED_PERSON]

	String Table INI redirects are now in the "Core.StringTable" section (rather than "/Script/Engine.Engine")

Change 3273483 on 2017/01/26 by [REDACTED_PERSON]

	Also makes a slight change to the ResavePackages commandlet to submit files marked for delete.

Change 3273568 on 2017/01/26 by [REDACTED_PERSON]

	Modifying changes made to SPluginWizard to have the plugin loading phase determined by the wizard's definition rather than from the first selected template.

Change 3273855 on 2017/01/26 by [REDACTED_PERSON]

Change 3274200 on 2017/01/26 by [REDACTED_PERSON]

	For IPluginWizardDefinition, temporarily adding function bodies to two methods instead of having them be pure virtual methods.

Change 3274317 on 2017/01/26 by [REDACTED_PERSON]

	Deleting a seemingly corrupted asset that was accidentially submitted

Change 3275072 on 2017/01/27 by [REDACTED_PERSON]

	Fix typo error
	Iterate all components, not only active one
	Force expand the Target Layers widget

Change 3275249 on 2017/01/27 by [REDACTED_PERSON]

	Color grading controls: Keep the vector ratio when changing the master slider

Change 3275282 on 2017/01/27 by [REDACTED_PERSON]

	Color grading controls: Cosmetic changes

Change 3275292 on 2017/01/27 by [REDACTED_PERSON]

	Make sure the build is called once when we import a staticmesh.

Change 3275430 on 2017/01/27 by [REDACTED_PERSON]

	Add some fbx automation tests
	- Import a mesh with no material
	- Import corrupted asset with no section in a LOD
	- Import morph targets
	- Materials name clash
	- Max Multimap material ordering

Change 3275683 on 2017/01/27 by [REDACTED_PERSON]

Change 3276237 on 2017/01/27 by [REDACTED_PERSON]

	Deleting a seemingly corrupted asset that was accidentially submitted

Change 3276266 on 2017/01/27 by [REDACTED_PERSON]

	Fix for accessing a potentially null pointer

Change 3277065 on 2017/01/30 by [REDACTED_PERSON]

	Move crash report temp files to saved config and cleanup on schedule.
	[UE-39506] - CrashReportClient ini folders are not cleaned when opening the editor

Change 3277236 on 2017/01/30 by [REDACTED_PERSON]

	Fix crash when cancelling SaveCurrentLevelAs

Change 3277409 on 2017/01/30 by [REDACTED_PERSON]

	Improved text rendering when the last resort font is missing

	The last resort font is no longer included in shipping builds, so this change makes some improvements to text rendering when it's missing.

	- The legacy font cache no longer tries to use the last resort font if it's not available (preventing warnings).
	- The Slate font renderer no longer tries to use the last resort font if it's not available.
	- Text shaping will use the last resort character if none of the available fonts can render a given character (likely because the last resort font is missing).
	- HarfBuzz shaped text now uses the fallback character correctly.

Change 3277749 on 2017/01/30 by [REDACTED_PERSON]

	Slate - Moving ESlateDrawEffect & ESlateBatchDrawFlag over to be enum class, found cases where users were improperly assuming the enum order, and so now it won't be possible to just treat an int32 or a bool as the draw effect value.

	Core - Adding EnumHasAllFlags and EnumHasAnyFlags, templated functions to make it easier to check for the existance of a flag on enum classes.

Change 3277805 on 2017/01/30 by [REDACTED_PERSON]

	Rendering - Changing some LoadModuleChecked calls to GetModuleChecked, as these calls are not happening on the main thread and are not safe to make.

Change 3277914 on 2017/01/30 by [REDACTED_PERSON]

	Fix Niagara slate style warning on startup

Change 3278058 on 2017/01/30 by [REDACTED_PERSON]

	Fixed compile error

Change 3278132 on 2017/01/30 by [REDACTED_PERSON]

	Fixed compile error

Change 3278133 on 2017/01/30 by [REDACTED_PERSON]

	Fixed compile errors

Change 3278186 on 2017/01/30 by [REDACTED_PERSON]

	Fixed compile error

Change 3278525 on 2017/01/30 by [REDACTED_PERSON]

	Fixed compile error

Change 3278534 on 2017/01/30 by [REDACTED_PERSON]

	Automation - Clearing up several warnings/errors with automation results, trying to get Automation Tests to at least yellow before integration.

Change 3278941 on 2017/01/31 by [REDACTED_PERSON]

	Fixing a build warning due to build team refactor.

Change 3278949 on 2017/01/31 by [REDACTED_PERSON]

	Fixing incrmenetal build issues.

Change 3278953 on 2017/01/31 by [REDACTED_PERSON]

	Fixing some incrmental linux build issues.

Change 3278964 on 2017/01/31 by [REDACTED_PERSON]

	FIxing more incremental build issues.

Change 3279256 on 2017/01/31 by [REDACTED_PERSON]

	Instead of getting the Landscape Actor, call GetLandscapeProxy so all case are handled, either proxy or landscape actor

Change 3279270 on 2017/01/31 by [REDACTED_PERSON]

	re-updating the automation test pool

diff --git a/Engine/Build/Mac/RadioEffectUnit/RadioEffectUnitVersion.h b/Engine/Build/Mac/RadioEffectUnit/RadioEffectUnitVersion.h
--- a/Engine/Build/Mac/RadioEffectUnit/RadioEffectUnitVersion.h
+++ b/Engine/Build/Mac/RadioEffectUnit/RadioEffectUnitVersion.h
@@ -2,9 +2,6 @@
  * Copyright 1998-2017 Epic Games, Inc. All Rights Reserved.
  */
 
-#ifndef __RadioEffectUnitVersion_h__
-#define __RadioEffectUnitVersion_h__
+#pragma once
 
 #define kRadioEffectUnitVersion 0x00010000
-
-#endif
diff --git a/Engine/Config/BaseEditorPerProjectUserSettings.ini b/Engine/Config/BaseEditorPerProjectUserSettings.ini
--- a/Engine/Config/BaseEditorPerProjectUserSettings.ini
+++ b/Engine/Config/BaseEditorPerProjectUserSettings.ini
@@ -77,7 +77,7 @@ SelectorColor=(R=0.701f,G=0.225f,B=0.003f,A=1.000000)
 
 LogBackgroundColor=(R=0.047059,G=0.047059,B=0.047059,A=1.000000)
 LogSelectionBackgroundColor=(R=0.132868,G=0.132868,B=0.132868,A=1.000000)
-LogNormalColor=(R=0.401978,G=0.401978,B=0.401978,A=1.000000)
+LogNormalColor=(R=1.0,G=1.0,B=1.0,A=1.000000)
 LogCommandColor=(R=0.033105,G=0.723055,B=0.033105,A=1.000000)
 LogWarningColor=(R=0.496933,G=0.496933,B=0.057805,A=1.000000)
 LogErrorColor=(R=0.723055,G=0.000000,B=0.000000,A=1.000000)
@@ -720,6 +720,7 @@ DefaultBrushRadius=128
 +Node=(TargetClass=ActorComponent TargetEvent="ReceiveTick")
 +Node=(TargetClass=GameplayAbility TargetEvent="K2_ActivateAbility")
 +Node=(TargetClass=GameplayAbility TargetEvent="K2_OnEndAbility")
++Node=(TargetClass=UserWidget TargetEvent="PreConstruct")
 +Node=(TargetClass=UserWidget TargetEvent="Construct")
 +Node=(TargetClass=UserWidget TargetEvent="Tick")
 +Node=(TargetClass=AnimInstance TargetEvent="BlueprintUpdateAnimation")
@@ -752,6 +753,9 @@ DefaultBrushRadius=128
 [WidgetTemplatesExpanded]
 Common=True
 
+[DetailCustomWidgetExpansion]
+LandscapeEditorObject=LandscapeEditorObject.Target Layers.TargetLayers
+
 [LevelSequenceEditor SequencerSettings]
 bKeyInterpPropertiesOnly=true
 bShowRangeSlider=true
diff --git a/Engine/Config/BaseEngine.ini b/Engine/Config/BaseEngine.ini
--- a/Engine/Config/BaseEngine.ini
+++ b/Engine/Config/BaseEngine.ini
@@ -1421,7 +1421,6 @@ bUseStreamingPause=false
 
 ; UMG Checkbox State Rename
 
-+ActiveClassRedirects=(OldClassName="ESlateCheckBoxState", NewClassName="/Script/SlateCore.ECheckBoxState")
 +ActiveClassRedirects=(OldClassName="ECheckBoxState",NewClassName="/Script/SlateCore.ECheckBoxState")
 
 ; UMG Moving FAnchor to Slate
@@ -2818,6 +2817,7 @@ ConnectionRetryDelay=2
 [CrashReportClient]
 bHideLogFilesOption=false
 bIsAllowedToCloseWithoutSending=true
+CrashConfigPurgeDays=2
 
 [SteamVR.Settings]
 HMDWornMovementThreshold = 50.0
@@ -2833,7 +2833,11 @@ r.MeshReductionModule="QuadricMeshReduction"
 +InternalOnlyPaths=(Path="/Engine/NotForLicensees")
 +InternalOnlyClasses=/Script/VREditor.VREditorBaseUserWidget
 +InternalOnlyClasses=/Script/LevelSequence.LevelSequenceBurnIn
+
 [EngineSessionManager]
 UseWatchdogMTBF=false
+AllowWatchdogDialogs=false
+AllowWatchdogDetectHangs=false
+WatchdogMinimumHangSeconds=120
 
 [/Script/Niagara.NiagaraSettings]
\ No newline at end of file
diff --git a/Engine/Plugins/Developer/GitSourceControl/Source/GitSourceControl/Private/SGitSourceControlSettings.cpp b/Engine/Plugins/Developer/GitSourceControl/Source/GitSourceControl/Private/SGitSourceControlSettings.cpp
--- a/Engine/Plugins/Developer/GitSourceControl/Source/GitSourceControl/Private/SGitSourceControlSettings.cpp
+++ b/Engine/Plugins/Developer/GitSourceControl/Source/GitSourceControl/Private/SGitSourceControlSettings.cpp
@@ -13,10 +13,13 @@
 #include "Widgets/Input/SButton.h"
 #include "Widgets/Input/SCheckBox.h"
 #include "EditorStyleSet.h"
-#include "GitSourceControlModule.h"
-#include "GitSourceControlUtils.h"
 #include "SSeparator.h"
 #include "SMultiLineEditableTextBox.h"
+#include "Framework/Notifications/NotificationManager.h"
+#include "Widgets/Notifications/SNotificationList.h"
+#include "SourceControlOperations.h"
+#include "GitSourceControlModule.h"
+#include "GitSourceControlUtils.h"
 
 #define LOCTEXT_NAMESPACE "SGitSourceControlSettings"
 
@@ -243,6 +246,11 @@ void SGitSourceControlSettings::Construct(const FArguments& InArgs)
 	];
 }
 
+SGitSourceControlSettings::~SGitSourceControlSettings()
+{
+	RemoveInProgressNotification();
+}
+
 FText SGitSourceControlSettings::GetBinaryPathText() const
 {
 	FGitSourceControlModule& GitSourceControl = FModuleManager::LoadModuleChecked<FGitSourceControlModule>("GitSourceControl");
@@ -285,49 +293,145 @@ EVisibility SGitSourceControlSettings::CanInitializeGitRepository() const
 FReply SGitSourceControlSettings::OnClickedInitializeGitRepository()
 {
 	FGitSourceControlModule& GitSourceControl = FModuleManager::LoadModuleChecked<FGitSourceControlModule>("GitSourceControl");
-	TArray<FString> InfoMessages;
-	TArray<FString> ErrorMessages;
 	const FString& PathToGitBinary = GitSourceControl.AccessSettings().GetBinaryPath();
 	const FString PathToGameDir = FPaths::ConvertRelativePathToFull(FPaths::GameDir());
+	TArray<FString> InfoMessages;
+	TArray<FString> ErrorMessages;
+
+	// 1. Synchronous (very quick) "git init" operation: initialize a Git local repository with a .git/ subdirectory
 	GitSourceControlUtils::RunCommand(TEXT("init"), PathToGitBinary, PathToGameDir, TArray<FString>(), TArray<FString>(), InfoMessages, ErrorMessages);
+
 	// Check the new repository status to enable connection
 	GitSourceControl.GetProvider().CheckGitAvailability();
 	if(GitSourceControl.GetProvider().IsEnabled())
 	{
+		// List of files to add to Source Control (.uproject, Config/, Content/, Source/ files and .gitignore if any)
 		TArray<FString> ProjectFiles;
-		ProjectFiles.Add(FPaths::GetCleanFilename(FPaths::GetProjectFilePath()));
-		ProjectFiles.Add(FPaths::GetCleanFilename(FPaths::GameConfigDir()));
-		ProjectFiles.Add(FPaths::GetCleanFilename(FPaths::GameContentDir()));
-		if(FPaths::DirectoryExists(FPaths::GameSourceDir()))
+		ProjectFiles.Add(FPaths::GetProjectFilePath());
+		ProjectFiles.Add(FPaths::GameConfigDir());
+		ProjectFiles.Add(FPaths::GameContentDir());
+		if (FPaths::DirectoryExists(FPaths::GameSourceDir()))
 		{
-			ProjectFiles.Add(FPaths::GetCleanFilename(FPaths::GameSourceDir()));
+			ProjectFiles.Add(FPaths::GameSourceDir());
 		}
 		if(bAutoCreateGitIgnore)
 		{
-			// Create a standard ".gitignore" file with common patterns for a typical Blueprint & C++ project
-			const FString Filename = FString::Printf(TEXT("%s.gitignore"), *PathToGameDir);
+			// 2. Create a standard ".gitignore" file with common patterns for a typical Blueprint & C++ project
+			const FString GitIgnoreFilename = FPaths::Combine(FPaths::GameDir(), TEXT(".gitignore"));
 			const FString GitIgnoreContent = TEXT("Binaries\nDerivedDataCache\nIntermediate\nSaved\n*.VC.db\n*.opensdf\n*.opendb\n*.sdf\n*.sln\n*.suo\n*.xcodeproj\n*.xcworkspace");
-			if(FFileHelper::SaveStringToFile(GitIgnoreContent, *Filename, FFileHelper::EEncodingOptions::ForceUTF8WithoutBOM))
+			if(FFileHelper::SaveStringToFile(GitIgnoreContent, *GitIgnoreFilename, FFileHelper::EEncodingOptions::ForceUTF8WithoutBOM))
 			{
-				ProjectFiles.Add(TEXT(".gitignore"));
+				ProjectFiles.Add(GitIgnoreFilename);
 			}
 		}
-		// Add .uproject, Config/, Content/ and Source/ files (and .gitignore if any)
-		GitSourceControlUtils::RunCommand(TEXT("add"), PathToGitBinary, PathToGameDir, TArray<FString>(), ProjectFiles, InfoMessages, ErrorMessages);
-		if(bAutoInitialCommit)
-		{
-			// optionnal initial git commit with custom message
-			TArray<FString> Parameters;
-			FString ParamCommitMsg = TEXT("--message=\"");
-			ParamCommitMsg += InitialCommitMessage.ToString();
-			ParamCommitMsg += TEXT("\"");
-			Parameters.Add(ParamCommitMsg);
-			GitSourceControlUtils::RunCommit(PathToGitBinary, PathToGameDir, Parameters, TArray<FString>(), InfoMessages, ErrorMessages);
-		}
+
+		// 3. Add files to Source Control: launch an asynchronous MarkForAdd operation
+		LaunchMarkForAddOperation(ProjectFiles);
+
+		// 4. The CheckIn will follow, at completion of the MarkForAdd operation
 	}
 	return FReply::Handled();
 }
 
+// Launch an asynchronous "MarkForAdd" operation and start an ongoing notification
+void SGitSourceControlSettings::LaunchMarkForAddOperation(const TArray<FString>& InFiles)
+{
+	FGitSourceControlModule& GitSourceControl = FModuleManager::LoadModuleChecked<FGitSourceControlModule>("GitSourceControl");
+	TSharedRef<FMarkForAdd, ESPMode::ThreadSafe> MarkForAddOperation = ISourceControlOperation::Create<FMarkForAdd>();
+	ECommandResult::Type Result = GitSourceControl.GetProvider().Execute(MarkForAddOperation, InFiles, EConcurrency::Asynchronous, FSourceControlOperationComplete::CreateSP(this, &SGitSourceControlSettings::OnSourceControlOperationComplete));
+	if (Result == ECommandResult::Succeeded)
+	{
+		DisplayInProgressNotification(MarkForAddOperation);
+	}
+	else
+	{
+		DisplayFailureNotification(MarkForAddOperation);
+	}
+}
+
+// Launch an asynchronous "CheckIn" operation and start another ongoing notification
+void SGitSourceControlSettings::LaunchCheckInOperation()
+{
+	TSharedRef<FCheckIn, ESPMode::ThreadSafe> CheckInOperation = ISourceControlOperation::Create<FCheckIn>();
+	CheckInOperation->SetDescription(InitialCommitMessage);
+	FGitSourceControlModule& GitSourceControl = FModuleManager::LoadModuleChecked<FGitSourceControlModule>("GitSourceControl");
+	ECommandResult::Type Result = GitSourceControl.GetProvider().Execute(CheckInOperation, TArray<FString>(), EConcurrency::Asynchronous, FSourceControlOperationComplete::CreateSP(this, &SGitSourceControlSettings::OnSourceControlOperationComplete));
+	if (Result == ECommandResult::Succeeded)
+	{
+		DisplayInProgressNotification(CheckInOperation);
+	}
+	else
+	{
+		DisplayFailureNotification(CheckInOperation);
+	}
+}
+
+/// Delegate called when a source control operation has completed: launch the next one and manage notifications
+void SGitSourceControlSettings::OnSourceControlOperationComplete(const FSourceControlOperationRef& InOperation, ECommandResult::Type InResult)
+{
+	RemoveInProgressNotification();
+
+	// Report result with a notification
+	if (InResult == ECommandResult::Succeeded)
+	{
+		DisplaySuccessNotification(InOperation);
+	}
+	else
+	{
+		DisplayFailureNotification(InOperation);
+	}
+
+	if ((InOperation->GetName() == "MarkForAdd") && (InResult == ECommandResult::Succeeded) && bAutoInitialCommit)
+	{
+		// 4. optional initial Asynchronous commit with custom message: launch a "CheckIn" Operation
+		LaunchCheckInOperation();
+	}
+}
+
+
+// Display an ongoing notification during the whole operation
+void SGitSourceControlSettings::DisplayInProgressNotification(const FSourceControlOperationRef& InOperation)
+{
+	FNotificationInfo Info(InOperation->GetInProgressString());
+	Info.bFireAndForget = false;
+	Info.ExpireDuration = 0.0f;
+	Info.FadeOutDuration = 1.0f;
+	OperationInProgressNotification = FSlateNotificationManager::Get().AddNotification(Info);
+	if (OperationInProgressNotification.IsValid())
+	{
+		OperationInProgressNotification.Pin()->SetCompletionState(SNotificationItem::CS_Pending);
+	}
+}
+
+// Remove the ongoing notification at the end of the operation
+void SGitSourceControlSettings::RemoveInProgressNotification()
+{
+	if (OperationInProgressNotification.IsValid())
+	{
+		OperationInProgressNotification.Pin()->ExpireAndFadeout();
+		OperationInProgressNotification.Reset();
+	}
+}
+
+// Display a temporary success notification at the end of the operation
+void SGitSourceControlSettings::DisplaySuccessNotification(const FSourceControlOperationRef& InOperation)
+{
+	const FText NotificationText = FText::Format(LOCTEXT("InitialCommit_Success", "{0} operation was successfull!"), FText::FromName(InOperation->GetName()));
+	FNotificationInfo Info(NotificationText);
+	Info.bUseSuccessFailIcons = true;
+	Info.Image = FEditorStyle::GetBrush(TEXT("NotificationList.SuccessImage"));
+	FSlateNotificationManager::Get().AddNotification(Info);
+}
+
+// Display a temporary failure notification at the end of the operation
+void SGitSourceControlSettings::DisplayFailureNotification(const FSourceControlOperationRef& InOperation)
+{
+	const FText NotificationText = FText::Format(LOCTEXT("InitialCommit_Failure", "Error: {0} operation failed!"), FText::FromName(InOperation->GetName()));
+	FNotificationInfo Info(NotificationText);
+	Info.ExpireDuration = 8.0f;
+	FSlateNotificationManager::Get().AddNotification(Info);
+}
+
 void SGitSourceControlSettings::OnCheckedCreateGitIgnore(ECheckBoxState NewCheckedState)
 {
 	bAutoCreateGitIgnore = (NewCheckedState == ECheckBoxState::Checked);
diff --git a/Engine/Plugins/Developer/GitSourceControl/Source/GitSourceControl/Private/SGitSourceControlSettings.h b/Engine/Plugins/Developer/GitSourceControl/Source/GitSourceControl/Private/SGitSourceControlSettings.h
--- a/Engine/Plugins/Developer/GitSourceControl/Source/GitSourceControl/Private/SGitSourceControlSettings.h
+++ b/Engine/Plugins/Developer/GitSourceControl/Source/GitSourceControl/Private/SGitSourceControlSettings.h
@@ -7,6 +7,9 @@
 #include "Input/Reply.h"
 #include "Widgets/DeclarativeSyntaxSupport.h"
 #include "Widgets/SCompoundWidget.h"
+#include "SlateFwd.h"
+#include "ISourceControlOperation.h"
+#include "ISourceControlProvider.h"
 
 enum class ECheckBoxState : uint8;
 
@@ -22,6 +25,8 @@ public:
 
 	void Construct(const FArguments& InArgs);
 
+	~SGitSourceControlSettings();
+
 private:
 
 	/** Delegate to get binary path from settings */
@@ -47,4 +52,19 @@ private:
 	void OnInitialCommitMessageCommited(const FText& InText, ETextCommit::Type InCommitType);
 	FText GetInitialCommitMessage() const;
 	FText InitialCommitMessage;
+
+	/** Launch initial asynchronous add and commit operations */
+	void LaunchMarkForAddOperation(const TArray<FString>& InFiles);
+	void LaunchCheckInOperation();
+
+	/** Delegate called when a source control operation has completed */
+	void OnSourceControlOperationComplete(const FSourceControlOperationRef& InOperation, ECommandResult::Type InResult);
+
+	/** Asynchronous operation progress notifications */
+	TWeakPtr<SNotificationItem> OperationInProgressNotification;
+	
+	void DisplayInProgressNotification(const FSourceControlOperationRef& InOperation);
+	void RemoveInProgressNotification();
+	void DisplaySuccessNotification(const FSourceControlOperationRef& InOperation);
+	void DisplayFailureNotification(const FSourceControlOperationRef& InOperation);
 };
diff --git a/Engine/Plugins/Editor/PluginBrowser/Source/PluginBrowser/Private/DefaultPluginWizardDefinition.cpp b/Engine/Plugins/Editor/PluginBrowser/Source/PluginBrowser/Private/DefaultPluginWizardDefinition.cpp
--- a/Engine/Plugins/Editor/PluginBrowser/Source/PluginBrowser/Private/DefaultPluginWizardDefinition.cpp
+++ b/Engine/Plugins/Editor/PluginBrowser/Source/PluginBrowser/Private/DefaultPluginWizardDefinition.cpp
@@ -101,6 +101,30 @@ bool FDefaultPluginWizardDefinition::GetPluginIconPath(FString& OutIconPath) con
 	return GetTemplateIconPath(CurrentTemplateDefinition.ToSharedRef(), OutIconPath);
 }
 
+EHostType::Type FDefaultPluginWizardDefinition::GetPluginModuleDescriptor() const
+{
+	EHostType::Type ModuleDescriptorType = EHostType::Developer;
+
+	if (CurrentTemplateDefinition.IsValid())
+	{
+		ModuleDescriptorType = CurrentTemplateDefinition->ModuleDescriptorType;
+	}
+
+	return ModuleDescriptorType;
+}
+
+ELoadingPhase::Type FDefaultPluginWizardDefinition::GetPluginLoadingPhase() const
+{
+	ELoadingPhase::Type Phase = ELoadingPhase::Default;
+
+	if (CurrentTemplateDefinition.IsValid())
+	{
+		Phase = CurrentTemplateDefinition->LoadingPhase;
+	}
+
+	return Phase;
+}
+
 bool FDefaultPluginWizardDefinition::GetTemplateIconPath(TSharedRef<FPluginTemplateDescription> Template, FString& OutIconPath) const
 {
 	bool bRequiresDefaultIcon = false;
diff --git a/Engine/Plugins/Editor/PluginBrowser/Source/PluginBrowser/Private/DefaultPluginWizardDefinition.h b/Engine/Plugins/Editor/PluginBrowser/Source/PluginBrowser/Private/DefaultPluginWizardDefinition.h
--- a/Engine/Plugins/Editor/PluginBrowser/Source/PluginBrowser/Private/DefaultPluginWizardDefinition.h
+++ b/Engine/Plugins/Editor/PluginBrowser/Source/PluginBrowser/Private/DefaultPluginWizardDefinition.h
@@ -24,6 +24,8 @@ public:
 	virtual FText GetInstructions() const override;
 
 	virtual bool GetPluginIconPath(FString& OutIconPath) const override;
+	virtual EHostType::Type GetPluginModuleDescriptor() const override;
+	virtual ELoadingPhase::Type GetPluginLoadingPhase() const override;
 	virtual bool GetTemplateIconPath(TSharedRef<FPluginTemplateDescription> InTemplate, FString& OutIconPath) const override;
 	virtual FString GetFolderForSelection() const override;
 	// End IPluginWizardDefinition interface
diff --git a/Engine/Plugins/Editor/PluginBrowser/Source/PluginBrowser/Private/SNewPluginWizard.cpp b/Engine/Plugins/Editor/PluginBrowser/Source/PluginBrowser/Private/SNewPluginWizard.cpp
--- a/Engine/Plugins/Editor/PluginBrowser/Source/PluginBrowser/Private/SNewPluginWizard.cpp
+++ b/Engine/Plugins/Editor/PluginBrowser/Source/PluginBrowser/Private/SNewPluginWizard.cpp
@@ -512,27 +512,15 @@ FReply SNewPluginWizard::OnCreatePluginClicked()
 	{
 		bHasModules = true;
 	}
-	
-	EHostType::Type ModuleDescriptorType;
-	ELoadingPhase::Type LoadingPhase;
-
-	TArray<TSharedPtr<FPluginTemplateDescription>> SelectedTemplates = PluginWizardDefinition->GetSelectedTemplates();
-	// @todo For now assume there is one module type and all templates share it
-	if(SelectedTemplates.Num() > 0 )
-	{
-		ModuleDescriptorType = SelectedTemplates[0]->ModuleDescriptorType;
-		LoadingPhase = SelectedTemplates[0]->LoadingPhase;
-	}
-	else
-	{
-		// Default to runtime
-		ModuleDescriptorType = EHostType::Runtime;
-		LoadingPhase = ELoadingPhase::Default;
-	}
 
 	// Save descriptor file as .uplugin file
 	const FString UPluginFilePath = GetPluginFilenameWithPath();
+
+	const EHostType::Type ModuleDescriptorType = PluginWizardDefinition->GetPluginModuleDescriptor();
+	const ELoadingPhase::Type LoadingPhase = PluginWizardDefinition->GetPluginLoadingPhase();
+
 	bSucceeded = bSucceeded && WritePluginDescriptor(AutoPluginName, UPluginFilePath, PluginWizardDefinition->CanContainContent(), bHasModules, ModuleDescriptorType, LoadingPhase);
+	
 
 	// Main plugin dir
 	const FString BasePluginFolder = GetPluginDestinationPath().ToString();
diff --git a/Engine/Plugins/Editor/PluginBrowser/Source/PluginBrowser/Private/SPluginCategoryTree.cpp b/Engine/Plugins/Editor/PluginBrowser/Source/PluginBrowser/Private/SPluginCategoryTree.cpp
--- a/Engine/Plugins/Editor/PluginBrowser/Source/PluginBrowser/Private/SPluginCategoryTree.cpp
+++ b/Engine/Plugins/Editor/PluginBrowser/Source/PluginBrowser/Private/SPluginCategoryTree.cpp
@@ -191,9 +191,9 @@ void SPluginCategoryTree::RebuildAndFilterCategoryTree()
 		{
 			TreeView->SetSelection(InstalledCategory);
 		}
-		else if(RootCategories.Num() > 0 && RootCategories[0]->SubCategories.Num() > 0)
+		else if(RootCategories.Num() > 0)
 		{
-			TreeView->SetSelection(RootCategories[0]->SubCategories[0]);
+			TreeView->SetSelection(RootCategories[0]);
 		}
 	}
 }
diff --git a/Engine/Plugins/Editor/PluginBrowser/Source/PluginBrowser/Public/IPluginWizardDefinition.h b/Engine/Plugins/Editor/PluginBrowser/Source/PluginBrowser/Public/IPluginWizardDefinition.h
--- a/Engine/Plugins/Editor/PluginBrowser/Source/PluginBrowser/Public/IPluginWizardDefinition.h
+++ b/Engine/Plugins/Editor/PluginBrowser/Source/PluginBrowser/Public/IPluginWizardDefinition.h
@@ -92,6 +92,20 @@ public:
 	/** Gets the icon path for the current template selection. Returns true if the plugin requires a default icon */
 	virtual bool GetPluginIconPath(FString& OutIconPath) const = 0;
 
+	/** Gets the ModuleDescriptor for the plugin based on the selection  */
+	virtual EHostType::Type GetPluginModuleDescriptor() const 
+	{
+		// @todo: make this a pure virtual function
+		return EHostType::Runtime;
+	}
+
+	/** Gets the LoadingPhase for the plugin based on the selection */
+	virtual ELoadingPhase::Type GetPluginLoadingPhase() const
+	{
+		// @todo: make this a pure virtual function
+		return ELoadingPhase::Default;
+	}
+
 	/** Gets the icon path for the specified template. Returns true if it requires a default icon */
 	virtual bool GetTemplateIconPath(TSharedRef<FPluginTemplateDescription> InTemplate, FString& OutIconPath) const = 0;
 
diff --git a/Engine/Plugins/MovieScene/LevelSequenceEditor/Source/LevelSequenceEditor/Private/CinematicViewport/SCinematicTransportRange.cpp b/Engine/Plugins/MovieScene/LevelSequenceEditor/Source/LevelSequenceEditor/Private/CinematicViewport/SCinematicTransportRange.cpp
--- a/Engine/Plugins/MovieScene/LevelSequenceEditor/Source/LevelSequenceEditor/Private/CinematicViewport/SCinematicTransportRange.cpp
+++ b/Engine/Plugins/MovieScene/LevelSequenceEditor/Source/LevelSequenceEditor/Private/CinematicViewport/SCinematicTransportRange.cpp
@@ -103,7 +103,7 @@ int32 SCinematicTransportRange::OnPaint(const FPaintArgs& Args, const FGeometry&
 {
 	ISequencer* Sequencer = GetSequencer();
 
-	ESlateDrawEffect::Type DrawEffects = bParentEnabled ? ESlateDrawEffect::None : ESlateDrawEffect::DisabledEffect;
+	const ESlateDrawEffect DrawEffects = bParentEnabled ? ESlateDrawEffect::None : ESlateDrawEffect::DisabledEffect;
 	if (!Sequencer)
 	{
 		return LayerId;
diff --git a/Engine/Plugins/NetcodeUnitTest/NetcodeUnitTest/Source/NetcodeUnitTest/Private/UI/SLogWidget.cpp b/Engine/Plugins/NetcodeUnitTest/NetcodeUnitTest/Source/NetcodeUnitTest/Private/UI/SLogWidget.cpp
--- a/Engine/Plugins/NetcodeUnitTest/NetcodeUnitTest/Source/NetcodeUnitTest/Private/UI/SLogWidget.cpp
+++ b/Engine/Plugins/NetcodeUnitTest/NetcodeUnitTest/Source/NetcodeUnitTest/Private/UI/SLogWidget.cpp
@@ -304,21 +304,12 @@ void SLogWidget::Construct(const FArguments& Args)
 					.AutoWidth()
 					[
 						SAssignNew(ArrayAddNew(AutoCloseWidgets), SCheckBox)
-#if TARGET_UE4_CL >= CL_CHECKBOXSTATE
 						.IsChecked((bAutoClose ? ECheckBoxState::Checked : ECheckBoxState::Unchecked))
 						.OnCheckStateChanged_Lambda(
 						[&](ECheckBoxState NewAutoCloseState)
 						{
 							bAutoClose = (NewAutoCloseState == ECheckBoxState::Checked);
 						})
-#else
-						.IsChecked((bAutoClose ? ESlateCheckBoxState::Checked : ESlateCheckBoxState::Unchecked))
-						.OnCheckStateChanged_Lambda(
-						[&](ESlateCheckBoxState::Type NewAutoCloseState)
-						{
-							bAutoClose = (NewAutoCloseState == ESlateCheckBoxState::Checked);
-						})
-#endif
 					]
 				)
 				+ConditionalSlot(!Args._bStatusWidget,
@@ -358,19 +349,11 @@ void SLogWidget::Construct(const FArguments& Args)
 							{
 								return CanAutoScroll(GetActiveTabInfo().ToSharedRef());
 							})
-#if TARGET_UE4_CL >= CL_CHECKBOXSTATE
 						.IsChecked(ECheckBoxState::Checked)
 						.OnCheckStateChanged_Lambda(
 							[&](ECheckBoxState NewAutoScrollState)
 							{
 								bAutoScroll = (NewAutoScrollState == ECheckBoxState::Checked);
-#else
-						.IsChecked(ESlateCheckBoxState::Checked)
-						.OnCheckStateChanged_Lambda(
-							[&](ESlateCheckBoxState::Type NewAutoScrollState)
-							{
-								bAutoScroll = (NewAutoScrollState == ESlateCheckBoxState::Checked);
-#endif
 
 								// Scroll to the end, if autoscroll was just re-enabled
 								if (bAutoScroll)
@@ -419,21 +402,12 @@ void SLogWidget::Construct(const FArguments& Args)
 					.AutoWidth()
 					[
 						SAssignNew(ArrayAddNew(DeveloperWidgets), SCheckBox)
-#if TARGET_UE4_CL >= CL_CHECKBOXSTATE
 						.IsChecked(ECheckBoxState::Unchecked)
 						.OnCheckStateChanged_Lambda(
 							[&](ECheckBoxState NewDeveloperState)
 							{
 								OnDeveloperClicked.ExecuteIfBound(NewDeveloperState == ECheckBoxState::Checked);
 							})
-#else
-						.IsChecked(ESlateCheckBoxState::Unchecked)
-						.OnCheckStateChanged_Lambda(
-							[&](ESlateCheckBoxState::Type NewDeveloperState)
-							{
-								OnDeveloperClicked.ExecuteIfBound(NewDeveloperState == ESlateCheckBoxState::Checked);
-							})
-#endif
 					]
 				)
 
diff --git a/Engine/Plugins/NetcodeUnitTest/NetcodeUnitTest/Source/NetcodeUnitTest/Public/NetcodeUnitTest.h b/Engine/Plugins/NetcodeUnitTest/NetcodeUnitTest/Source/NetcodeUnitTest/Public/NetcodeUnitTest.h
--- a/Engine/Plugins/NetcodeUnitTest/NetcodeUnitTest/Source/NetcodeUnitTest/Public/NetcodeUnitTest.h
+++ b/Engine/Plugins/NetcodeUnitTest/NetcodeUnitTest/Source/NetcodeUnitTest/Public/NetcodeUnitTest.h
@@ -50,7 +50,6 @@
 #define CL_DEPRECATENEW		2425600
 #define CL_DEPRECATEDEL		2400883
 #define CL_FNETWORKVERSION	2384479
-#define CL_CHECKBOXSTATE	2384008
 
 
 
diff --git a/Engine/Plugins/Online/OnlineSubsystemOculus/Source/Private/OculusNetConnection.cpp b/Engine/Plugins/Online/OnlineSubsystemOculus/Source/Private/OculusNetConnection.cpp
--- a/Engine/Plugins/Online/OnlineSubsystemOculus/Source/Private/OculusNetConnection.cpp
+++ b/Engine/Plugins/Online/OnlineSubsystemOculus/Source/Private/OculusNetConnection.cpp
@@ -1,7 +1,5 @@
 // Copyright 1998-2017 Epic Games, Inc. All Rights Reserved.
 
-#pragma once
-
 #include "OculusNetConnection.h"
 #include "OnlineSubsystemOculusPrivate.h"
 
diff --git a/Engine/Plugins/Online/OnlineSubsystemOculus/Source/Private/OculusNetDriver.cpp b/Engine/Plugins/Online/OnlineSubsystemOculus/Source/Private/OculusNetDriver.cpp
--- a/Engine/Plugins/Online/OnlineSubsystemOculus/Source/Private/OculusNetDriver.cpp
+++ b/Engine/Plugins/Online/OnlineSubsystemOculus/Source/Private/OculusNetDriver.cpp
@@ -1,7 +1,5 @@
 // Copyright 2016 Oculus VR, LLC All Rights reserved.
 
-#pragma once
-
 #include "OculusNetDriver.h"
 #include "OnlineSubsystemOculusPrivate.h"
 
diff --git a/Engine/Plugins/Tests/EditorTests/Source/EditorTests/Private/UnrealEd/BlueprintAutomationTests.cpp b/Engine/Plugins/Tests/EditorTests/Source/EditorTests/Private/UnrealEd/BlueprintAutomationTests.cpp
--- a/Engine/Plugins/Tests/EditorTests/Source/EditorTests/Private/UnrealEd/BlueprintAutomationTests.cpp
+++ b/Engine/Plugins/Tests/EditorTests/Source/EditorTests/Private/UnrealEd/BlueprintAutomationTests.cpp
@@ -341,7 +341,7 @@ public:
 	static void UnloadBlueprint(UBlueprint* const BlueprintObj, bool bForceFlush = false)
 	{
 		// have to grab the blueprint's package before we move it to the transient package
-		UPackage* const OldPackage = Cast<UPackage>(BlueprintObj->GetOutermost());
+		UPackage* const OldPackage = BlueprintObj->GetOutermost();
 
 		UPackage* const TransientPackage = GetTransientPackage();
 		if (OldPackage == TransientPackage)
@@ -498,7 +498,7 @@ public:
 			return false;
 		}
 
-		UPackage* const BlueprintPackage = Cast<UPackage>(BlueprintObj->GetOutermost());
+		UPackage* const BlueprintPackage = BlueprintObj->GetOutermost();
 		// compiling the blueprint will inherently dirty the package, but if there 
 		// weren't any changes to save before, there shouldn't be after
 		bool const bStartedWithUnsavedChanges = (BlueprintPackage != nullptr) ? BlueprintPackage->IsDirty() : true;
@@ -650,7 +650,7 @@ public:
 
 		FString SavePath = FString::Printf(TEXT("%sTemp-%u-%s"), *TempDir, GenTempUid(), *FPaths::GetCleanFilename(BlueprintObj->GetName()));
 
-		UPackage* const AssetPackage = Cast<UPackage>(BlueprintObj->GetOutermost());
+		UPackage* const AssetPackage = BlueprintObj->GetOutermost();
 		return UPackage::SavePackage(AssetPackage, NULL, RF_Standalone, *SavePath, GWarn);
 	}
 
@@ -1338,7 +1338,7 @@ bool FBlueprintRenameAndCloneTest::RunTest(const FString& BlueprintAssetPath)
 		if (!bIsAlreadyLoaded)
 		{
 			// store the original package so we can manually invalidate it after the move
-			UPackage* const OriginalPackage = Cast<UPackage>(OriginalBlueprint->GetOutermost());
+			UPackage* const OriginalPackage = OriginalBlueprint->GetOutermost();
 
 			FString const BlueprintName = OriginalBlueprint->GetName();
 			UPackage* TempPackage = FBlueprintAutomationTestUtilities::CreateTempPackage(BlueprintName);
diff --git a/Engine/Plugins/Tests/EditorTests/Source/EditorTests/Private/UnrealEd/EditorAssetAutomationTests.cpp b/Engine/Plugins/Tests/EditorTests/Source/EditorTests/Private/UnrealEd/EditorAssetAutomationTests.cpp
--- a/Engine/Plugins/Tests/EditorTests/Source/EditorTests/Private/UnrealEd/EditorAssetAutomationTests.cpp
+++ b/Engine/Plugins/Tests/EditorTests/Source/EditorTests/Private/UnrealEd/EditorAssetAutomationTests.cpp
@@ -561,6 +561,13 @@ bool FAssetEditorTest::RunTest(const FString& Parameters)
 		}
 	}
 
+	// No skeleton was loaded, just load one the first one found
+	// This is only used to verify we can create assets that rely on skeletons.
+	if (!FirstSkeleton && AllSkeletons.Num())
+	{
+		FirstSkeleton = CastChecked<USkeleton>(AllSkeletons[0].GetAsset());
+	}
+
 	//Check to see if we need to enable Behavior trees.
 	bool bEnabledBehaviorTrees = false;
 	bool bBehaviorTreeNewAssetsEnabled = false;
@@ -624,7 +631,6 @@ bool FAssetEditorTest::RunTest(const FString& Parameters)
 	ASSET_TEST_CREATE(UBlueprint, UBlueprintMacroFactory, MPL, FactoryInst->ParentClass = AActor::StaticClass();)
 	ASSET_TEST_CREATE(UCameraAnim, UCameraAnimFactory, CA, )
 	ASSET_TEST_CREATE(UCurveBase, UCurveFactory, C, FactoryInst->CurveClass = UCurveFloat::StaticClass();)
-	FModuleManager::Get().LoadModule(TEXT("GameplayAbilities"));
 	UClass* GameplayAbilityClass = StaticLoadClass(UObject::StaticClass(), NULL, TEXT("GameplayAbilities.GameplayAbilitySet"), NULL, LOAD_None, NULL);
 	if (GameplayAbilityClass)
 	{
diff --git a/Engine/Plugins/Tests/EditorTests/Source/EditorTests/Private/UnrealEd/FbxAutomationTests.cpp b/Engine/Plugins/Tests/EditorTests/Source/EditorTests/Private/UnrealEd/FbxAutomationTests.cpp
--- a/Engine/Plugins/Tests/EditorTests/Source/EditorTests/Private/UnrealEd/FbxAutomationTests.cpp
+++ b/Engine/Plugins/Tests/EditorTests/Source/EditorTests/Private/UnrealEd/FbxAutomationTests.cpp
@@ -157,6 +157,7 @@ bool FFbxImportAssetsAutomationTest::RunTest(const FString& Parameters)
 		switch (TestPlan->Action)
 		{
 			case EFBXTestPlanActionType::Import:
+			case EFBXTestPlanActionType::ImportReload:
 			{
 				//Create a factory and set the options
 				UFbxFactory* FbxFactory = NewObject<UFbxFactory>(UFbxFactory::StaticClass());
@@ -172,6 +173,70 @@ bool FFbxImportAssetsAutomationTest::RunTest(const FString& Parameters)
 
 				//Import the test object
 				ImportedObjects = AssetToolsModule.Get().ImportAssets(CurFileToImport, ImportAssetPath, FbxFactory);
+				if (TestPlan->Action == EFBXTestPlanActionType::ImportReload)
+				{
+					TArray<FString> FullAssetPaths;
+					
+					TArray<FAssetData> ImportedAssets;
+					AssetRegistryModule.Get().GetAssetsByPath(FName(*ImportAssetPath), ImportedAssets, true);
+					for (const FAssetData& AssetData : ImportedAssets)
+					{
+						UObject *Asset = AssetData.GetAsset();
+						if (Asset != nullptr)
+						{
+							if (ImportedObjects.Contains(Asset))
+							{
+								FullAssetPaths.Add(Asset->GetPathName());
+							}
+							FString PackageName = Asset->GetOutermost()->GetPathName();
+							Asset->MarkPackageDirty();
+							UPackage::SavePackage(Asset->GetOutermost(), Asset, RF_Standalone, *FPackageName::LongPackageNameToFilename(PackageName, FPackageName::GetAssetPackageExtension()), GError, nullptr, false, true, SAVE_NoError);
+						}
+					}
+					for (const FAssetData& AssetData : ImportedAssets)
+					{
+						UPackage *Package = AssetData.GetPackage();
+						if (Package != nullptr)
+						{
+							for (FObjectIterator It; It; ++It)
+							{
+								UObject* ExistingObject = *It;
+								if ((ExistingObject->GetOutermost() == Package) )
+								{
+									ExistingObject->ClearFlags(RF_Standalone | RF_Public);
+									ExistingObject->RemoveFromRoot();
+									ExistingObject->MarkPendingKill();
+								}
+						
+							}
+						}
+					}
+					CollectGarbage(GARBAGE_COLLECTION_KEEPFLAGS);
+					
+					ImportedObjects.Empty();
+
+					for (const FAssetData& AssetData : ImportedAssets)
+					{
+						UPackage *Package = AssetData.GetPackage();
+						if (Package != nullptr)
+						{
+							if (!Package->IsFullyLoaded())
+							{
+								Package->FullyLoad();
+							}
+						}
+					}
+					
+					//Set back the importObjects
+					for (FString PathName : FullAssetPaths)
+					{
+						UStaticMesh* FoundMesh = LoadObject<UStaticMesh>(NULL, *PathName, NULL, LOAD_Quiet | LOAD_NoWarn);
+						if (FoundMesh)
+						{
+							ImportedObjects.Add(FoundMesh);
+						}
+					}
+				}
 
 				//Add the just imported object to the global array use for reimport test
 				for (UObject* ImportObject : ImportedObjects)
@@ -472,6 +537,67 @@ bool FFbxImportAssetsAutomationTest::RunTest(const FString& Parameters)
 				}
 			}
 			break;
+			case Material_Slot_Imported_Name:
+			{
+				if (ExpectedResult.ExpectedPresetsDataInteger.Num() < 1)
+				{
+					ExecutionInfo.Errors.Add(FString::Printf(TEXT("%s expected result need 1 integer data (Expected material slot index)"),
+						*GetFormatedMessageErrorInTestData(CleanFilename, TestPlan->TestPlanName, TEXT("Material_Slot_Imported_Name"), ExpectedResultIndex)));
+					break;
+				}
+				if (ExpectedResult.ExpectedPresetsDataString.Num() < 1)
+				{
+					ExecutionInfo.Errors.Add(FString::Printf(TEXT("%s expected result need 1 string data (Expected material imported name for the specified slot index)"),
+						*GetFormatedMessageErrorInTestData(CleanFilename, TestPlan->TestPlanName, TEXT("Material_Slot_Imported_Name"), ExpectedResultIndex)));
+					break;
+				}
+				int32 MeshMaterialNumber = INDEX_NONE;
+				int32 MaterialSlotIndex = ExpectedResult.ExpectedPresetsDataInteger[0];
+				FString ExpectedMaterialImportedName = ExpectedResult.ExpectedPresetsDataString[0];
+				FString MaterialImportedName = TEXT("");
+				bool BadSlotIndex = false;
+				if (ImportedObjects.Num() > 0)
+				{
+					UObject *Object = ImportedObjects[0];
+					if (Object->IsA(UStaticMesh::StaticClass()))
+					{
+						UStaticMesh *Mesh = Cast<UStaticMesh>(Object);
+						if (Mesh->StaticMaterials.Num() <= MaterialSlotIndex)
+						{
+							BadSlotIndex = true;
+							MeshMaterialNumber = Mesh->StaticMaterials.Num();
+						}
+						else
+						{
+							MaterialImportedName = Mesh->StaticMaterials[MaterialSlotIndex].ImportedMaterialSlotName.ToString();
+						}
+					}
+					else if (Object->IsA(USkeletalMesh::StaticClass()))
+					{
+						USkeletalMesh *Mesh = Cast<USkeletalMesh>(Object);
+						if (Mesh->Materials.Num() <= MaterialSlotIndex)
+						{
+							BadSlotIndex = true;
+							MeshMaterialNumber = Mesh->Materials.Num();
+						}
+						else
+						{
+							MaterialImportedName = Mesh->Materials[MaterialSlotIndex].ImportedMaterialSlotName.ToString();
+						}
+					}
+				}
+				if(BadSlotIndex)
+				{
+					ExecutionInfo.Errors.Add(FString::Printf(TEXT("%s->%s: Error in the test data, Material_Slot_Imported_Name material slot index [%d] is invalid. Expect something smaller then %d which is the mesh material number"),
+						*CleanFilename, *(TestPlan->TestPlanName), MaterialSlotIndex, MeshMaterialNumber));
+				}
+				else if (MaterialImportedName != ExpectedMaterialImportedName)
+				{
+					ExecutionInfo.Errors.Add(FString::Printf(TEXT("%s [Material slot index %d has a materials imported name %s but expected %s]"),
+						*GetFormatedMessageErrorInExpectedResult(*CleanFilename, *(TestPlan->TestPlanName), TEXT("Material_Slot_Imported_Name"), ExpectedResultIndex), MaterialSlotIndex, *MaterialImportedName, *ExpectedMaterialImportedName));
+				}
+			}
+			break;
 			case Vertex_Number:
 			{
 				if (ExpectedResult.ExpectedPresetsDataInteger.Num() < 1)
@@ -984,6 +1110,93 @@ bool FFbxImportAssetsAutomationTest::RunTest(const FString& Parameters)
 				}
 			}
 			break;
+			case Mesh_LOD_Section_Material_Imported_Name:
+			{
+				if (ExpectedResult.ExpectedPresetsDataInteger.Num() < 2 || ExpectedResult.ExpectedPresetsDataString.Num() < 1)
+				{
+					ExecutionInfo.Errors.Add(FString::Printf(TEXT("%s expected result need 2 integer data and 1 string(LOD index, section index and Expected material name)"),
+						*GetFormatedMessageErrorInTestData(CleanFilename, TestPlan->TestPlanName, TEXT("Mesh_LOD_Section_Material_Imported_Name"), ExpectedResultIndex)));
+					break;
+				}
+				int32 LODIndex = ExpectedResult.ExpectedPresetsDataInteger[0];
+				int32 SectionIndex = ExpectedResult.ExpectedPresetsDataInteger[1];
+				FString ExpectedMaterialName = ExpectedResult.ExpectedPresetsDataString[0];
+				int32 LODNumber = 0;
+				int32 SectionNumber = 0;
+				bool BadLodIndex = false;
+				bool BadSectionIndex = false;
+				FString MaterialName;
+				if (ImportedObjects.Num() > 0)
+				{
+					UObject *Object = ImportedObjects[0];
+					if (Object->IsA(UStaticMesh::StaticClass()))
+					{
+						UStaticMesh *Mesh = Cast<UStaticMesh>(Object);
+						LODNumber = Mesh->GetNumLODs();
+						if (LODIndex < 0 || LODIndex >= LODNumber)
+						{
+							BadLodIndex = true;
+						}
+						else
+						{
+							SectionNumber = Mesh->RenderData->LODResources[LODIndex].Sections.Num();
+							if (SectionIndex < 0 || SectionIndex >= SectionNumber)
+							{
+								BadSectionIndex = true;
+							}
+							else
+							{
+								int32 MaterialIndex = Mesh->RenderData->LODResources[LODIndex].Sections[SectionIndex].MaterialIndex;
+								if (MaterialIndex >= 0 && MaterialIndex < Mesh->StaticMaterials.Num())
+								{
+									MaterialName = Mesh->StaticMaterials[MaterialIndex].ImportedMaterialSlotName.ToString();
+								}
+							}
+						}
+					}
+					else if (Object->IsA(USkeletalMesh::StaticClass()))
+					{
+						USkeletalMesh *Mesh = Cast<USkeletalMesh>(Object);
+						LODNumber = Mesh->GetResourceForRendering()->LODModels.Num();
+						if (LODIndex < 0 || LODIndex >= LODNumber)
+						{
+							BadLodIndex = true;
+						}
+						else
+						{
+							SectionNumber = Mesh->GetResourceForRendering()->LODModels[LODIndex].Sections.Num();
+							if (SectionIndex < 0 || SectionIndex >= SectionNumber)
+							{
+								BadSectionIndex = true;
+							}
+							else
+							{
+								int32 MaterialIndex = Mesh->GetResourceForRendering()->LODModels[LODIndex].Sections[SectionIndex].MaterialIndex;
+								if (MaterialIndex >= 0 && MaterialIndex < Mesh->Materials.Num())
+								{
+									MaterialName = Mesh->Materials[MaterialIndex].ImportedMaterialSlotName.ToString();
+								}
+							}
+						}
+					}
+				}
+				if (BadLodIndex)
+				{
+					ExecutionInfo.Errors.Add(FString::Printf(TEXT("%s->%s: Error in the test data, Mesh_LOD_Section_Material_Imported_Name LOD index [%d] is invalid. Expect LODIndex between 0 and %d which is the mesh LOD number"),
+						*CleanFilename, *(TestPlan->TestPlanName), LODIndex, LODNumber));
+				}
+				else if (BadSectionIndex)
+				{
+					ExecutionInfo.Errors.Add(FString::Printf(TEXT("%s->%s: Error in the test data, Mesh_LOD_Section_Material_Imported_Name Section index [%d] is invalid. Expect Section Index between 0 and %d which is the mesh LOD section number"),
+						*CleanFilename, *(TestPlan->TestPlanName), LODIndex, SectionNumber));
+				}
+				else if (MaterialName.Compare(ExpectedMaterialName) != 0)
+				{
+					ExecutionInfo.Errors.Add(FString::Printf(TEXT("%s [LOD index %d Section index %d contain import material name (%s) but expected name (%s)]"),
+						*GetFormatedMessageErrorInExpectedResult(*CleanFilename, *(TestPlan->TestPlanName), TEXT("Mesh_LOD_Section_Material_Imported_Name"), ExpectedResultIndex), LODIndex, SectionIndex, *MaterialName, *ExpectedMaterialName));
+				}
+			}
+			break;
 			case LOD_UV_Channel_Number:
 			{
 				if (ExpectedResult.ExpectedPresetsDataInteger.Num() < 2)
@@ -1130,8 +1343,37 @@ bool FFbxImportAssetsAutomationTest::RunTest(const FString& Parameters)
 			};
 			ExpectedResultIndex++;
 		}
-		if (TestPlan->bDeleteFolderAssets)
+		if (TestPlan->bDeleteFolderAssets || TestPlan->Action == EFBXTestPlanActionType::ImportReload)
 		{
+			//When doing an import-reload we have to destroy the package since it was save
+			//But when we just have everything in memory a garbage collection pass is enough to
+			//delete assets.
+			if (TestPlan->Action != EFBXTestPlanActionType::ImportReload)
+			{
+				TArray<FAssetData> ImportedAssets;
+				AssetRegistryModule.Get().GetAssetsByPath(FName(*ImportAssetPath), ImportedAssets, true);
+				for (const FAssetData& AssetData : ImportedAssets)
+				{
+					UPackage *Package = AssetData.GetPackage();
+					if (Package != nullptr)
+					{
+						for (FObjectIterator It; It; ++It)
+						{
+							UObject* ExistingObject = *It;
+							if ((ExistingObject->GetOutermost() == Package))
+							{
+								ExistingObject->ClearFlags(RF_Standalone | RF_Public);
+								ExistingObject->RemoveFromRoot();
+								ExistingObject->MarkPendingKill();
+							}
+
+						}
+					}
+				}
+				CollectGarbage(GARBAGE_COLLECTION_KEEPFLAGS);
+			}
+
+			//Make sure there is no more asset under "Engine\Content\FbxEditorAutomationOut" folder
 			GlobalImportedObjects.Empty();
 			TArray<FAssetData> AssetsToDelete;
 			AssetRegistryModule.Get().GetAssetsByPath(FName(*PackagePath), AssetsToDelete, true);
diff --git a/Engine/Plugins/Tests/ScreenshotTools/Source/ScreenshotTools/Private/CompareScreenshotsCommandlet.cpp b/Engine/Plugins/Tests/ScreenshotTools/Source/ScreenshotTools/Private/CompareScreenshotsCommandlet.cpp
--- a/Engine/Plugins/Tests/ScreenshotTools/Source/ScreenshotTools/Private/CompareScreenshotsCommandlet.cpp
+++ b/Engine/Plugins/Tests/ScreenshotTools/Source/ScreenshotTools/Private/CompareScreenshotsCommandlet.cpp
@@ -24,7 +24,6 @@ int32 UCompareScreenshotsCommandlet::Main(FString const& Params)
 	IScreenShotToolsModule& ScreenShotModule = FModuleManager::LoadModuleChecked<IScreenShotToolsModule>("ScreenShotComparisonTools");
 	IScreenShotManagerPtr ScreenshotManager = ScreenShotModule.GetScreenShotManager();
 
-	ScreenshotManager->CompareScreensotsAsync().Wait();
 	ScreenshotManager->ExportScreensotsAsync(ScreenshotExportPath).Wait();
 
 	return 0;
diff --git a/Engine/Shaders/SlateElementPixelShader.usf b/Engine/Shaders/SlateElementPixelShader.usf
--- a/Engine/Shaders/SlateElementPixelShader.usf
+++ b/Engine/Shaders/SlateElementPixelShader.usf
@@ -16,6 +16,11 @@ float2 GammaValues;
 /** Per element params */
 float4 ShaderParams;
 
+#if !USE_MATERIALS
+/** Per element option to invert the alpha, mostly only useful when dealing with data coming from the scene renderer, which writes out inverse opacity. */
+float InvertAlpha;
+#endif
+
 /** Texture for the element */
 Texture2D ElementTexture;
 SamplerState ElementTextureSampler;
@@ -307,6 +312,12 @@ float4 Main( VertexToPixelInterpolants VIn ) : SV_Target0
 	// gamma correct
 	OutColor.rgb = GammaCorrect(OutColor.rgb);
 
+#if !USE_MATERIALS
+    // Deal with cases where the incoming alpha is inverted like if you're dealing with render targets coming
+    // from the scene renderer.
+    OutColor.a = lerp(OutColor.a, 1 - OutColor.a, InvertAlpha);
+#endif
+
 #if DEBUG_CLIPPING
 	if (any(ClipTest < 0))
 	{
diff --git a/Engine/Shaders/SlateVertexShader.usf b/Engine/Shaders/SlateVertexShader.usf
--- a/Engine/Shaders/SlateVertexShader.usf
+++ b/Engine/Shaders/SlateVertexShader.usf
@@ -1,6 +1,5 @@
 // Copyright 1998-2017 Epic Games, Inc. All Rights Reserved.
 
-
 #include "Common.usf"
 #include "GammaCorrectionCommon.usf"
 #include "SlateShaderCommon.usf"
@@ -15,15 +14,17 @@ VertexToPixelInterpolants Main(
 	in float2 InPosition : ATTRIBUTE2,
 	in float2 InClipOrigin : ATTRIBUTE3,
 	in float4 InClipExtents : ATTRIBUTE4,
-	in float4 InColor : ATTRIBUTE5
+	in float4 InColor : ATTRIBUTE5,
+	in uint2 PixelSize : ATTRIBUTE6
 #if USE_SLATE_INSTANCING
-	, in float4 InstanceParam : ATTRIBUTE6
+	, in float4 InstanceParam : ATTRIBUTE7
 #endif
 	)
 {
 	VertexToPixelInterpolants VOut = (VertexToPixelInterpolants)0;
 
-	float4 WorldPosition = float4(InPosition.xy,0,1);
+    float4 WorldPosition = float4(InPosition.xy, 0, 1);
+
 	// save off the original input position value, Line AA will require this
 	VOut.OriginalPosition = WorldPosition;
 
@@ -39,68 +40,70 @@ VertexToPixelInterpolants Main(
 	VOut.Color = FinalVertexColor;
 	VOut.TextureCoordinates[0] = InTextureCoordinates;
 
-#if USE_MATERIALS
+#if USE_MATERIALS 
 	FMaterialVertexParameters VertexParameters = (FMaterialVertexParameters)0;
 	VertexParameters.WorldPosition = WorldPosition.xyz;
 
 	VertexParameters.VertexColor = FinalVertexColor;
 
-	#if NUM_MATERIAL_TEXCOORDS_VERTEX
+#if NUM_MATERIAL_TEXCOORDS_VERTEX
 
 #if !USE_SLATE_INSTANCING
-		float4 InstanceParam = float4(0, 0, 0, 0);
+	// pass along local screen space size 
+	// To look this up in a material you use uv channel 3
+	float4 InstanceParam = float4(0,0, PixelSize);
 #endif
-		float2 UVArrays[8] =
-		{
-			InTextureCoordinates.xy,
-			InTextureCoordinates.zw,
-			InstanceParam.xy,
-			InstanceParam.zw,
-			InMaterialTextureCoordinates,
-			InClipOrigin,
-			InClipExtents.xy,
-			InClipExtents.zw
-		};
-
-		UNROLL
-		for( int CoordinateIndex = 0; CoordinateIndex < NUM_MATERIAL_TEXCOORDS_VERTEX; CoordinateIndex++ )
-		{
-			VertexParameters.TexCoords[CoordinateIndex] = UVArrays[CoordinateIndex];
-		}
-
-		#if HAS_MATERIAL_TEXCOORDS
-
-			float2 CustomizedUVs[NUM_SLATE_TEXCOORDS];
-			GetMaterialCustomizedUVs(VertexParameters, CustomizedUVs);
-
-			#if ES2_PROFILE
-				#if NUM_SLATE_TEXCOORDS == 1
-					VOut.TextureCoordinates[0].xy = CustomizedUVs[0];
-				#elif NUM_SLATE_TEXCOORDS >= 2
-					VOut.TextureCoordinates[0].xy = CustomizedUVs[0];
-					VOut.TextureCoordinates[0].zw = CustomizedUVs[1];
-				#endif
-			#else
-				UNROLL
-				for( int CoordinateIndex = 0; CoordinateIndex < NUM_SLATE_TEXCOORDS; CoordinateIndex++ )
+	float2 UVArrays[8] =
+	{
+		InTextureCoordinates.xy, //uv channel 0
+		InTextureCoordinates.zw, //uv channel 1
+		InstanceParam.xy, //uv channel 2
+		InstanceParam.zw, //uv channel 3
+		InMaterialTextureCoordinates, //uv channel 4
+		InClipOrigin, //uv channel 5
+		InClipExtents.xy, //uv channel 6
+		InClipExtents.zw //uv channel 7
+	};
+
+	UNROLL
+	for( int CoordinateIndex = 0; CoordinateIndex < NUM_MATERIAL_TEXCOORDS_VERTEX; CoordinateIndex++ )
+	{
+		VertexParameters.TexCoords[CoordinateIndex] = UVArrays[CoordinateIndex];
+	}
+
+	#if HAS_MATERIAL_TEXCOORDS
+
+		float2 CustomizedUVs[NUM_SLATE_TEXCOORDS];
+		GetMaterialCustomizedUVs(VertexParameters, CustomizedUVs);
+
+		#if ES2_PROFILE
+			#if NUM_SLATE_TEXCOORDS == 1
+				VOut.TextureCoordinates[0].xy = CustomizedUVs[0];
+			#elif NUM_SLATE_TEXCOORDS >= 2
+				VOut.TextureCoordinates[0].xy = CustomizedUVs[0];
+				VOut.TextureCoordinates[0].zw = CustomizedUVs[1];
+			#endif
+		#else
+			UNROLL
+			for( int CoordinateIndex = 0; CoordinateIndex < NUM_SLATE_TEXCOORDS; CoordinateIndex++ )
+			{
+				FLATTEN
+				if (CoordinateIndex % 2)
 				{
-					FLATTEN
-					if (CoordinateIndex % 2)
-					{
-						VOut.TextureCoordinates[CoordinateIndex / 2].zw = CustomizedUVs[CoordinateIndex];
-					}
-					else
-					{
-						VOut.TextureCoordinates[CoordinateIndex / 2].xy = CustomizedUVs[CoordinateIndex];
-					}
+					VOut.TextureCoordinates[CoordinateIndex / 2].zw = CustomizedUVs[CoordinateIndex];
 				}
-			#endif
-		#endif // HAS_MATERIAL_TEXCOORDS
-	#endif // NUM_MATERIAL_TEXCOORDS_VERTEX
+				else
+				{
+					VOut.TextureCoordinates[CoordinateIndex / 2].xy = CustomizedUVs[CoordinateIndex];
+				}
+			}
+		#endif
+	#endif // HAS_MATERIAL_TEXCOORDS
+#endif // NUM_MATERIAL_TEXCOORDS_VERTEX
 
-	#if HAS_SCREEN_POSITION
-		WorldPosition.xyz = GetMaterialWorldPositionOffset(VertexParameters);
-	#endif
+#if HAS_SCREEN_POSITION
+	WorldPosition.xyz = GetMaterialWorldPositionOffset(VertexParameters);
+#endif
 
 #endif // USE_MATERIALS
 
diff --git a/Engine/Source/Developer/AssetTools/Private/AssetRenameManager.cpp b/Engine/Source/Developer/AssetTools/Private/AssetRenameManager.cpp
--- a/Engine/Source/Developer/AssetTools/Private/AssetRenameManager.cpp
+++ b/Engine/Source/Developer/AssetTools/Private/AssetRenameManager.cpp
@@ -40,6 +40,7 @@
 #include "ObjectTools.h"
 #include "Interfaces/IMainFrameModule.h"
 #include "Kismet2/KismetEditorUtilities.h"
+#include "Kismet2/BlueprintEditorUtils.h"
 
 #define LOCTEXT_NAMESPACE "AssetRenameManager"
 
@@ -639,6 +640,12 @@ void FAssetRenameManager::RenameReferencingStringAssetReferences(const TArray<UP
 		for (auto* Object : ObjectsInPackage)
 		{
 			Object->Serialize(RenameSerializer);
+
+			if (UBlueprint* Blueprint = Cast<UBlueprint>(Object))
+			{
+				// Serialize may have dirtied the BP bytecode in some way
+				FBlueprintEditorUtils::MarkBlueprintAsModified(Blueprint);
+			}
 		}
 	}
 }
diff --git a/Engine/Source/Developer/AssetTools/Private/AssetTools.cpp b/Engine/Source/Developer/AssetTools/Private/AssetTools.cpp
--- a/Engine/Source/Developer/AssetTools/Private/AssetTools.cpp
+++ b/Engine/Source/Developer/AssetTools/Private/AssetTools.cpp
@@ -424,7 +424,7 @@ UObject* FAssetTools::CreateAsset(const FString& AssetName, const FString& Packa
 
 	UPackage* Pkg = CreatePackage(nullptr,*PackageName);
 	UObject* NewObj = nullptr;
-	EObjectFlags Flags = RF_Public|RF_Standalone;
+	EObjectFlags Flags = RF_Public|RF_Standalone|RF_Transactional;
 	if ( Factory )
 	{
 		NewObj = Factory->FactoryCreateNew(ClassToUse, Pkg, FName( *AssetName ), Flags, nullptr, GWarn, CallingContext);
@@ -1116,7 +1116,7 @@ TArray<UObject*> FAssetTools::ImportAssetsInternal(const TArray<FString>& Files,
 		// @todo import: gmp: show dialog in case of multiple matching factories
 		for(TObjectIterator<UClass> ClassIt; ClassIt; ++ClassIt)
 		{
-			if(!(*ClassIt)->IsChildOf(UFactory::StaticClass()) || ((*ClassIt)->HasAnyClassFlags(CLASS_Abstract)))
+			if(!(*ClassIt)->IsChildOf(UFactory::StaticClass()) || ((*ClassIt)->HasAnyClassFlags(CLASS_Abstract)) || (*ClassIt)->IsChildOf(USceneImportFactory::StaticClass()))
 			{
 				continue;
 			}
@@ -1225,6 +1225,8 @@ TArray<UObject*> FAssetTools::ImportAssetsInternal(const TArray<FString>& Files,
 	bool bDontOverwriteAny = false;
 	bool bDontReplaceAny = false;
 
+	TArray<UFactory*> UsedFactories;
+
 	// Now iterate over the input files and use the same factory object for each file with the same extension
 	for(int32 FileIdx = 0; FileIdx < FilesAndDestinations.Num(); ++FileIdx)
 	{
@@ -1286,8 +1288,12 @@ TArray<UObject*> FAssetTools::ImportAssetsInternal(const TArray<FString>& Files,
 		if(Factory != nullptr)
 		{
 			// Reset the 'Do you want to overwrite the existing object?' Yes to All / No to All prompt, to make sure the
-			// user gets a chance to select something
-			Factory->ResetState();
+			// user gets a chance to select something when the factory is first used during this import
+			if (!UsedFactories.Contains(Factory))
+			{
+				Factory->ResetState();
+				UsedFactories.AddUnique(Factory);
+			}
 
 			UClass* ImportAssetType = Factory->SupportedClass;
 			bool bImportSucceeded = false;
@@ -1458,7 +1464,7 @@ TArray<UObject*> FAssetTools::ImportAssetsInternal(const TArray<FString>& Files,
 			Factory->SetAutomatedAssetImportData(Params.ImportData);
 
 			ImportAssetType = Factory->ResolveSupportedClass();
-			UObject* Result = Factory->ImportObject(ImportAssetType, Pkg, FName(*Name), RF_Public | RF_Standalone, Filename, nullptr, bImportWasCancelled);
+			UObject* Result = Factory->ImportObject(ImportAssetType, Pkg, FName(*Name), RF_Public | RF_Standalone | RF_Transactional, Filename, nullptr, bImportWasCancelled);
 
 			Factory->SetAutomatedAssetImportData(nullptr);
 
diff --git a/Engine/Source/Developer/AssetTools/Private/AssetTypeActions/AssetTypeActions_ForceFeedbackEffect.cpp b/Engine/Source/Developer/AssetTools/Private/AssetTypeActions/AssetTypeActions_ForceFeedbackEffect.cpp
--- a/Engine/Source/Developer/AssetTools/Private/AssetTypeActions/AssetTypeActions_ForceFeedbackEffect.cpp
+++ b/Engine/Source/Developer/AssetTools/Private/AssetTypeActions/AssetTypeActions_ForceFeedbackEffect.cpp
@@ -24,7 +24,7 @@ void FAssetTypeActions_ForceFeedbackEffect::GetActions( const TArray<UObject*>&
 	MenuBuilder.AddMenuEntry(
 		LOCTEXT("ForceFeedbackEffect_PlayEffect", "Play"),
 		LOCTEXT("ForceFeedbackEffect_PlayEffectTooltip", "Plays the selected force feedback effect."),
-		FSlateIcon(FEditorStyle::GetStyleSetName(), "MediaAsset.AssetActions.Play"),
+		FSlateIcon(FEditorStyle::GetStyleSetName(), "MediaAsset.AssetActions.Play.Small"),
 		FUIAction(
 			FExecuteAction::CreateSP( this, &FAssetTypeActions_ForceFeedbackEffect::ExecutePlayEffect, Effects ),
 			FCanExecuteAction::CreateSP( this, &FAssetTypeActions_ForceFeedbackEffect::CanExecutePlayCommand, Effects )
@@ -34,7 +34,7 @@ void FAssetTypeActions_ForceFeedbackEffect::GetActions( const TArray<UObject*>&
 	MenuBuilder.AddMenuEntry(
 		LOCTEXT("ForceFeedbackEffect_StopEffect", "Stop"),
 		LOCTEXT("ForceFeedbackEffect_StopEffectTooltip", "Stops the selected force feedback effect."),
-		FSlateIcon(FEditorStyle::GetStyleSetName(), "MediaAsset.AssetActions.Stop"),
+		FSlateIcon(FEditorStyle::GetStyleSetName(), "MediaAsset.AssetActions.Stop.Small"),
 		FUIAction(
 			FExecuteAction::CreateSP( this, &FAssetTypeActions_ForceFeedbackEffect::ExecuteStopEffect, Effects ),
 			FCanExecuteAction()
@@ -144,10 +144,10 @@ TSharedPtr<SWidget> FAssetTypeActions_ForceFeedbackEffect::GetThumbnailOverlay(c
 	{
 		if (IsEffectPlaying(EffectList))
 		{
-			return FEditorStyle::GetBrush("MediaAsset.AssetActions.Stop");
+			return FEditorStyle::GetBrush("MediaAsset.AssetActions.Stop.Large");
 		}
 
-	return FEditorStyle::GetBrush("MediaAsset.AssetActions.Play");
+		return FEditorStyle::GetBrush("MediaAsset.AssetActions.Play.Large");
 	};
 
 	auto IsEnabledLambda = [this, EffectList]() -> bool
@@ -179,28 +179,43 @@ TSharedPtr<SWidget> FAssetTypeActions_ForceFeedbackEffect::GetThumbnailOverlay(c
 		return LOCTEXT("Thumbnail_PlayForceFeedbackToolTip", "Play selected force feedback effect");
 	};
 
-	return SNew(SBox)
+	TSharedRef<SBox> Box = SNew(SBox)
 		.HAlign(HAlign_Center)
 		.VAlign(VAlign_Center)
-		.Padding(FMargin(2))
+		.Padding(FMargin(2));
+
+	auto OnGetVisibilityLambda = [this, Box, EffectList]() -> EVisibility
+	{
+		if (Box->IsHovered() || IsEffectPlaying(EffectList))
+		{
+			return EVisibility::Visible;
+		}
+
+		return EVisibility::Hidden;
+	};
+
+	TSharedRef<SButton> BoxContent = SNew(SButton)
+		.ButtonStyle(FEditorStyle::Get(), "HoverHintOnly")
+		.ToolTipText_Lambda(OnToolTipTextLambda)
+		.Cursor(EMouseCursor::Default) // The outer widget can specify a DragHand cursor, so we need to override that here
+		.ForegroundColor(FSlateColor::UseForeground())		
+		.IsEnabled_Lambda(IsEnabledLambda)
+		.OnClicked_Lambda(OnClickedLambda)
+		.Visibility_Lambda(OnGetVisibilityLambda)
 		[
-			SNew(SButton)
-			.ButtonStyle(FEditorStyle::Get(), "HoverHintOnly")
-			.ToolTipText_Lambda(OnToolTipTextLambda)
-			.Cursor(EMouseCursor::Default) // The outer widget can specify a DragHand cursor, so we need to override that here
-			.ForegroundColor(FSlateColor::UseForeground())		
-			.IsEnabled_Lambda(IsEnabledLambda)
-			.OnClicked_Lambda(OnClickedLambda)
+			SNew(SBox)
+			.MinDesiredWidth(16)
+			.MinDesiredHeight(16)
 			[
-				SNew(SBox)
-				.MinDesiredWidth(16)
-				.MinDesiredHeight(16)
-				[
-					SNew(SImage)
-					.Image_Lambda(OnGetDisplayBrushLambda)
-				]
+				SNew(SImage)
+				.Image_Lambda(OnGetDisplayBrushLambda)
 			]
 		];
+
+	Box->SetContent(BoxContent);
+	Box->SetVisibility(EVisibility::Visible);
+
+	return Box;
 }
 
 bool FPreviewForceFeedbackEffect::IsTickable() const
diff --git a/Engine/Source/Developer/AutomationController/AutomationController.Build.cs b/Engine/Source/Developer/AutomationController/AutomationController.Build.cs
--- a/Engine/Source/Developer/AutomationController/AutomationController.Build.cs
+++ b/Engine/Source/Developer/AutomationController/AutomationController.Build.cs
@@ -22,7 +22,8 @@ namespace UnrealBuildTool.Rules
                     "MessageLog",
                     "Json",
                     "JsonUtilities",
-					"ScreenShotComparisonTools"
+					"ScreenShotComparisonTools",
+					"HTTP"
                 }
 			);
 
diff --git a/Engine/Source/Developer/AutomationController/Private/AutomationControllerManager.h b/Engine/Source/Developer/AutomationController/Private/AutomationControllerManager.h
--- a/Engine/Source/Developer/AutomationController/Private/AutomationControllerManager.h
+++ b/Engine/Source/Developer/AutomationController/Private/AutomationControllerManager.h
@@ -281,9 +281,14 @@ protected:
 	void ReportTestResults();
 
 	/**
-	* Create a json file that contains all of our test report data at /saved/automation/logs/AutomationReport-{CL}-{DateTime}.json
-	*/
-	void GenerateJsonTestPassSummary();
+	 * Create a json file that contains all of our test report data at /saved/automation/logs/AutomationReport-{CL}-{DateTime}.json
+	 */
+	void GenerateJsonTestPassSummary(FDateTime Timestamp);
+
+	/**
+	 * Generates a full html report of the testing, which may include links to images.  All of it will be bundled under a folder.
+	 */
+	void GenerateHtmlTestPassSummary(FDateTime Timestamp);
 
 	/**
 	* Updates the result value of a finished test.
diff --git a/Engine/Source/Developer/AutomationController/Private/AutomationControllerManger.cpp b/Engine/Source/Developer/AutomationController/Private/AutomationControllerManger.cpp
--- a/Engine/Source/Developer/AutomationController/Private/AutomationControllerManger.cpp
+++ b/Engine/Source/Developer/AutomationController/Private/AutomationControllerManger.cpp
@@ -21,6 +21,7 @@
 #include "JsonObjectConverter.h"
 #include "Misc/EngineVersion.h"
 #include "Misc/FileHelper.h"
+#include "PlatformHttp.h"
 
 #if WITH_EDITOR
 #include "Logging/MessageLog.h"
@@ -297,7 +298,7 @@ void FAutomationControllerManager::UpdateTestResultValue(FString TestName, EAuto
 	}
 }
 
-void FAutomationControllerManager::GenerateJsonTestPassSummary()
+void FAutomationControllerManager::GenerateJsonTestPassSummary(FDateTime Timestamp)
 {
 	if (!OurPassResults.TestInformation.Num())
 	{
@@ -307,7 +308,7 @@ void FAutomationControllerManager::GenerateJsonTestPassSummary()
 	TSharedPtr<FJsonObject> ReportJson = FJsonObjectConverter::UStructToJsonObject(SerializedPassResults);
 	if (ReportJson.IsValid())
 	{
-		FString ReportFileName = FString::Printf(TEXT("%s/AutomationReport-%d-%s.json"), *FPaths::AutomationLogDir(), FEngineVersion::Current().GetChangelist(), *FDateTime::Now().ToString());
+		FString ReportFileName = FString::Printf(TEXT("%s/AutomationReport-%d-%s.json"), *FPaths::AutomationLogDir(), FEngineVersion::Current().GetChangelist(), *Timestamp.ToString());
 		FArchive* ReportFileWriter = IFileManager::Get().CreateFileWriter(*ReportFileName);
 		if (ReportFileWriter != nullptr)
 		{
@@ -323,6 +324,87 @@ void FAutomationControllerManager::GenerateJsonTestPassSummary()
 	}
 }
 
+void FAutomationControllerManager::GenerateHtmlTestPassSummary(FDateTime Timestamp)
+{
+	if ( !OurPassResults.TestInformation.Num() )
+	{
+		return;
+	}
+
+	const FAutomatedTestPassResults SerializedPassResults = OurPassResults;
+
+	FString Header =
+		TEXT("<html>")
+		TEXT("<head>")
+		TEXT("<title>Test Results</title>")
+		TEXT("<head>")
+		TEXT("<body>")
+			TEXT("<h3>Test Results</h3>")
+			TEXT("<table>")
+				TEXT("<tr>")
+					TEXT("<th width='10%'>Step</th>")
+					TEXT("<th width='10%'>Result</th>")
+					TEXT("<th width='80%'>Remarks</th>")
+				TEXT("</tr>");
+
+	FString Results;
+	for ( const FAutomatedTestResult& Test : OurPassResults.TestInformation )
+	{
+		FString Row = TEXT("<tr>");
+		Row += FString::Printf(TEXT("<td>%s</td>"), *Test.TestName);
+		Row += FString::Printf(TEXT("<td>%s</td>"), *Test.TestResult);
+		Row += TEXT("<td>");
+
+		if ( Test.TestErrors.Num() > 0 )
+		{
+			Row += TEXT("<h4>Errors</h4>");
+			Row += TEXT("<ul>");
+			for ( const FString& LogItem : Test.TestErrors )
+			{
+				Row += FString::Printf(TEXT("<li>%s</li>"), *FPlatformHttp::HtmlEncode(LogItem));
+			}
+			Row += TEXT("</ul>");
+		}
+
+		if ( Test.TestWarnings.Num() > 0 )
+		{
+			Row += TEXT("<h4>Warnings</h4>");
+			Row += TEXT("<ul>");
+			for ( const FString& LogItem : Test.TestWarnings )
+			{
+				Row += FString::Printf(TEXT("<li>%s</li>"), *FPlatformHttp::HtmlEncode(LogItem));
+			}
+			Row += TEXT("</ul>");
+		}
+
+		if ( Test.TestInfo.Num() > 0 )
+		{
+			Row += TEXT("<h4>Info</h4>");
+			Row += TEXT("<ul>");
+			for ( const FString& LogItem : Test.TestInfo )
+			{
+				Row += FString::Printf(TEXT("<li>%s</li>"), *FPlatformHttp::HtmlEncode(LogItem));
+			}
+			Row += TEXT("</ul>");
+		}
+
+		Row += TEXT("</td>");
+		Row += TEXT("</tr>");
+
+		Results += Row;
+	}
+
+	FString Footer =
+			TEXT("</table>")
+		TEXT("</body>");
+
+	FString ReportFileName = FString::Printf(TEXT("%s/AutomationReport-%d-%s.html"), *FPaths::AutomationLogDir(), FEngineVersion::Current().GetChangelist(), *Timestamp.ToString());
+	if ( !FFileHelper::SaveStringToFile(Header + Results + Footer, *ReportFileName) )
+	{
+		GLog->Logf(ELogVerbosity::Error, TEXT("Test Report Html is invalid - report not generated."));
+	}
+}
+
 void FAutomationControllerManager::ExecuteNextTask( int32 ClusterIndex, OUT bool& bAllTestsCompleted )
 {
 	bool bTestThatRequiresMultiplePraticipantsHadEnoughParticipants = false;
@@ -538,9 +620,11 @@ void FAutomationControllerManager::ProcessResults()
 	}
 
 	// Write our the report of our automation pass to JSON. Then clean our array for the next pass.
-	GenerateJsonTestPassSummary();
-	OurPassResults.ClearAllEntries();
+	FDateTime Timestamp = FDateTime::Now();
+	GenerateJsonTestPassSummary(Timestamp);
+	GenerateHtmlTestPassSummary(Timestamp);
 
+	OurPassResults.ClearAllEntries();
 
 	SetControllerStatus( EAutomationControllerModuleState::Ready );
 }
@@ -841,17 +925,15 @@ void FAutomationControllerManager::HandleRunTestsReplyMessage( const FAutomation
 		TestResults.Warnings = Message.Warnings;
 
 		// Verify this device thought it was busy
-		TSharedPtr <IAutomationReport> Report = DeviceClusterManager.GetTest(ClusterIndex, DeviceIndex);
+		TSharedPtr<IAutomationReport> Report = DeviceClusterManager.GetTest(ClusterIndex, DeviceIndex);
 		check(Report.IsValid());
 
-		Report->SetResults(ClusterIndex,CurrentTestPass, TestResults);
-		
+		Report->SetResults(ClusterIndex, CurrentTestPass, TestResults);
 
 		// Gather all of the data relevant to this test for our json reporting.
 		CollectTestNotes(Report->GetDisplayName(), Message);
 		UpdateTestResultValue(Report->GetDisplayName(), TestResults.State);
 
-
 #if WITH_EDITOR
 		FMessageLog AutomationTestingLog("AutomationTestingLog");
 		AutomationTestingLog.Open();
diff --git a/Engine/Source/Developer/AutomationDriver/Private/AutomatedApplication.cpp b/Engine/Source/Developer/AutomationDriver/Private/AutomatedApplication.cpp
--- a/Engine/Source/Developer/AutomationDriver/Private/AutomatedApplication.cpp
+++ b/Engine/Source/Developer/AutomationDriver/Private/AutomatedApplication.cpp
@@ -109,13 +109,13 @@ public:
 		}
 	}
 
-	virtual void SetCustomShape(void* CursorHandle) override
+	virtual void SetTypeShape(EMouseCursor::Type InCursorType, void* CursorHandle) override
 	{
 		if (bAllowMessageHandling)
 		{
 			if (RealCursor.IsValid())
 			{
-				RealCursor->SetCustomShape(CursorHandle);
+				RealCursor->SetTypeShape(InCursorType, CursorHandle);
 			}
 		}
 	}
diff --git a/Engine/Source/Developer/DesktopWidgets/Private/Widgets/Input/SDirectoryPicker.cpp b/Engine/Source/Developer/DesktopWidgets/Private/Widgets/Input/SDirectoryPicker.cpp
--- a/Engine/Source/Developer/DesktopWidgets/Private/Widgets/Input/SDirectoryPicker.cpp
+++ b/Engine/Source/Developer/DesktopWidgets/Private/Widgets/Input/SDirectoryPicker.cpp
@@ -24,7 +24,6 @@ void SDirectoryPicker::Construct( const FArguments& InArgs )
 
 		+ SHorizontalBox::Slot()
 		.FillWidth(1.0f)
-		.Padding(4.0f, 0.0f, 0.0f, 0.0f)
 		[
 			SAssignNew(EditableTextBox, SEditableTextBox)
 			.Text(this, &SDirectoryPicker::GetFilePathText)
diff --git a/Engine/Source/Developer/DirectoryWatcher/Private/Mac/DirectoryWatchRequestMac.h b/Engine/Source/Developer/DirectoryWatcher/Private/Mac/DirectoryWatchRequestMac.h
--- a/Engine/Source/Developer/DirectoryWatcher/Private/Mac/DirectoryWatchRequestMac.h
+++ b/Engine/Source/Developer/DirectoryWatcher/Private/Mac/DirectoryWatchRequestMac.h
@@ -1,7 +1,6 @@
 // Copyright 1998-2017 Epic Games, Inc. All Rights Reserved.
 
-#ifndef _INC_DIRECTORYWATCHREQUESTMAC
-#define _INC_DIRECTORYWATCHREQUESTMAC
+#pragma once
 
 #include "Delegates/Delegate.h"
 #include "IDirectoryWatcher.h"
@@ -45,5 +44,3 @@ private:
 	void ProcessChanges( size_t EventCount, void* EventPaths, const FSEventStreamEventFlags EventFlags[] );
 	void Shutdown( void );
 };
-
-#endif // _INC_DIRECTORYWATCHREQUESTMAC
diff --git a/Engine/Source/Developer/DirectoryWatcher/Private/Mac/DirectoryWatcherMac.h b/Engine/Source/Developer/DirectoryWatcher/Private/Mac/DirectoryWatcherMac.h
--- a/Engine/Source/Developer/DirectoryWatcher/Private/Mac/DirectoryWatcherMac.h
+++ b/Engine/Source/Developer/DirectoryWatcher/Private/Mac/DirectoryWatcherMac.h
@@ -1,7 +1,6 @@
 // Copyright 1998-2017 Epic Games, Inc. All Rights Reserved.
 
-#ifndef _INC_DIRECTORYWATCHERMAC
-#define _INC_DIRECTORYWATCHERMAC
+#pragma once
 
 #include "IDirectoryWatcher.h"
 #include "Containers/Map.h"
@@ -30,5 +29,3 @@ public:
 };
 
 typedef FDirectoryWatcherMac FDirectoryWatcher;
-
-#endif // _INC_DIRECTORYWATCHERMAC
diff --git a/Engine/Source/Developer/FunctionalTesting/Private/AutomationBlueprintFunctionLibrary.cpp b/Engine/Source/Developer/FunctionalTesting/Private/AutomationBlueprintFunctionLibrary.cpp
--- a/Engine/Source/Developer/FunctionalTesting/Private/AutomationBlueprintFunctionLibrary.cpp
+++ b/Engine/Source/Developer/FunctionalTesting/Private/AutomationBlueprintFunctionLibrary.cpp
@@ -260,9 +260,9 @@ bool UAutomationBlueprintFunctionLibrary::TakeAutomationScreenshotInternal(const
 	FAutomationScreenshotTaker* TempObject = new FAutomationScreenshotTaker(Name, Options);
 #endif
 
-	static IConsoleVariable* HighResScreenshotDelay = IConsoleManager::Get().FindConsoleVariable(TEXT("r.HighResScreenshotDelay"));
-	check(HighResScreenshotDelay);
-	HighResScreenshotDelay->Set(10);
+	//static IConsoleVariable* HighResScreenshotDelay = IConsoleManager::Get().FindConsoleVariable(TEXT("r.HighResScreenshotDelay"));
+	//check(HighResScreenshotDelay);
+	//HighResScreenshotDelay->Set(10);
 
     if ( FPlatformProperties::HasFixedResolution() )
     {
diff --git a/Engine/Source/Developer/LogVisualizer/Private/SVisualLoggerTimelineBar.cpp b/Engine/Source/Developer/LogVisualizer/Private/SVisualLoggerTimelineBar.cpp
--- a/Engine/Source/Developer/LogVisualizer/Private/SVisualLoggerTimelineBar.cpp
+++ b/Engine/Source/Developer/LogVisualizer/Private/SVisualLoggerTimelineBar.cpp
@@ -117,7 +117,7 @@ int32 SVisualLoggerTimelineBar::OnPaint(const FPaintArgs& Args, const FGeometry&
 	static const FColor SelectedBarColor(255, 255, 255, 255);
 	const FSlateBrush* SelectedFillImage = FLogVisualizerStyle::Get().GetBrush("LogVisualizer.LogBar.Selected");
 
-	const ESlateDrawEffect::Type DrawEffects = ESlateDrawEffect::None;// bEnabled ? ESlateDrawEffect::None : ESlateDrawEffect::DisabledEffect;
+	const ESlateDrawEffect DrawEffects = ESlateDrawEffect::None;// bEnabled ? ESlateDrawEffect::None : ESlateDrawEffect::DisabledEffect;
 
 	TArray<float> ErrorTimes;
 	TArray<float> WarningTimes;
diff --git a/Engine/Source/Developer/LogVisualizer/Private/VisualLoggerTimeSliderController.cpp b/Engine/Source/Developer/LogVisualizer/Private/VisualLoggerTimeSliderController.cpp
--- a/Engine/Source/Developer/LogVisualizer/Private/VisualLoggerTimeSliderController.cpp
+++ b/Engine/Source/Developer/LogVisualizer/Private/VisualLoggerTimeSliderController.cpp
@@ -119,7 +119,7 @@ struct FDrawTickArgs
 	/** Start layer for elements */
 	int32 StartLayer;
 	/** Draw effects to apply */
-	ESlateDrawEffect::Type DrawEffects;
+	ESlateDrawEffect DrawEffects;
 	/** Whether or not to only draw major ticks */
 	bool bOnlyDrawMajorTicks;
 	/** Whether or not to mirror labels */
@@ -228,7 +228,7 @@ void FVisualLoggerTimeSliderController::DrawTicks( FSlateWindowElementList& OutD
 int32 FVisualLoggerTimeSliderController::OnPaintTimeSlider( bool bMirrorLabels, const FGeometry& AllottedGeometry, const FSlateRect& MyClippingRect, FSlateWindowElementList& OutDrawElements, int32 LayerId, const FWidgetStyle& InWidgetStyle, bool bParentEnabled ) const
 {
 	const bool bEnabled = bParentEnabled;
-	const ESlateDrawEffect::Type DrawEffects = bEnabled ? ESlateDrawEffect::None : ESlateDrawEffect::DisabledEffect;
+	const ESlateDrawEffect DrawEffects = bEnabled ? ESlateDrawEffect::None : ESlateDrawEffect::DisabledEffect;
 
 	TRange<float> LocalViewRange = TimeSliderArgs.ViewRange.Get();
 	const float LocalViewRangeMin = LocalViewRange.GetLowerBoundValue();
@@ -619,7 +619,7 @@ FReply FVisualLoggerTimeSliderController::OnMouseWheel( SWidget& WidgetOwner, co
 
 int32 FVisualLoggerTimeSliderController::OnPaintSectionView( const FGeometry& AllottedGeometry, const FSlateRect& MyClippingRect, FSlateWindowElementList& OutDrawElements, int32 LayerId, bool bEnabled, bool bDisplayTickLines, bool bDisplayScrubPosition  ) const
 {
-	const ESlateDrawEffect::Type DrawEffects = bEnabled ? ESlateDrawEffect::None : ESlateDrawEffect::DisabledEffect;
+	const ESlateDrawEffect DrawEffects = bEnabled ? ESlateDrawEffect::None : ESlateDrawEffect::DisabledEffect;
 
 	TRange<float> LocalViewRange = TimeSliderArgs.ViewRange.Get();
 	float LocalScrubPosition = TimeSliderArgs.ScrubPosition.Get();
diff --git a/Engine/Source/Developer/Profiler/Private/Widgets/SDataGraph.cpp b/Engine/Source/Developer/Profiler/Private/Widgets/SDataGraph.cpp
--- a/Engine/Source/Developer/Profiler/Private/Widgets/SDataGraph.cpp
+++ b/Engine/Source/Developer/Profiler/Private/Widgets/SDataGraph.cpp
@@ -301,7 +301,7 @@ int32 SDataGraph::OnPaint( const FPaintArgs& Args, const FGeometry& AllottedGeom
 
 	// Rendering info.
 	const bool bEnabled  = ShouldBeEnabled( bParentEnabled );
-	ESlateDrawEffect::Type DrawEffects = bEnabled ? ESlateDrawEffect::None : ESlateDrawEffect::DisabledEffect;
+	ESlateDrawEffect DrawEffects = bEnabled ? ESlateDrawEffect::None : ESlateDrawEffect::DisabledEffect;
 	const FSlateBrush* TimelineAreaBrush = FEditorStyle::GetBrush("Profiler.LineGraphArea");
 	const FSlateBrush* WhiteBrush = FEditorStyle::GetBrush("WhiteTexture");
 
diff --git a/Engine/Source/Developer/Profiler/Private/Widgets/SHistogram.cpp b/Engine/Source/Developer/Profiler/Private/Widgets/SHistogram.cpp
--- a/Engine/Source/Developer/Profiler/Private/Widgets/SHistogram.cpp
+++ b/Engine/Source/Developer/Profiler/Private/Widgets/SHistogram.cpp
@@ -35,7 +35,7 @@ int32 SHistogram::OnPaint( const FPaintArgs& Args, const FGeometry& AllottedGeom
 
 	// Rendering info.
 	const bool bEnabled  = ShouldBeEnabled( bParentEnabled );
-	ESlateDrawEffect::Type DrawEffects = bEnabled ? ESlateDrawEffect::None : ESlateDrawEffect::DisabledEffect;
+	ESlateDrawEffect DrawEffects = bEnabled ? ESlateDrawEffect::None : ESlateDrawEffect::DisabledEffect;
 	const FSlateBrush* TimelineAreaBrush = FEditorStyle::GetBrush("Profiler.LineGraphArea");
 	const FSlateBrush* WhiteBrush = FEditorStyle::GetBrush("WhiteTexture");
 	const FSlateBrush* FillImage  = FEditorStyle::GetBrush("TaskGraph.Mono");
diff --git a/Engine/Source/Developer/Profiler/Private/Widgets/SProfilerMiniView.cpp b/Engine/Source/Developer/Profiler/Private/Widgets/SProfilerMiniView.cpp
--- a/Engine/Source/Developer/Profiler/Private/Widgets/SProfilerMiniView.cpp
+++ b/Engine/Source/Developer/Profiler/Private/Widgets/SProfilerMiniView.cpp
@@ -97,7 +97,7 @@ int32 SProfilerMiniView::OnPaint( const FPaintArgs& Args, const FGeometry& Allot
 
 	// Rendering info.
 	const bool bEnabled = ShouldBeEnabled( bParentEnabled );
-	const ESlateDrawEffect::Type DrawEffects = bEnabled ? ESlateDrawEffect::None : ESlateDrawEffect::DisabledEffect;
+	const ESlateDrawEffect DrawEffects = bEnabled ? ESlateDrawEffect::None : ESlateDrawEffect::DisabledEffect;
 	const FSlateBrush* MiniViewArea = FEditorStyle::GetBrush( "Profiler.LineGraphArea" );
 	const FSlateBrush* WhiteBrush = FEditorStyle::GetBrush( "WhiteTexture" );
 
diff --git a/Engine/Source/Developer/Profiler/Private/Widgets/SProfilerMiniView.h b/Engine/Source/Developer/Profiler/Private/Widgets/SProfilerMiniView.h
--- a/Engine/Source/Developer/Profiler/Private/Widgets/SProfilerMiniView.h
+++ b/Engine/Source/Developer/Profiler/Private/Widgets/SProfilerMiniView.h
@@ -164,7 +164,7 @@ class SProfilerMiniView : public SCompoundWidget
 	/** Holds current state provided by OnPaint function, used to simplify drawing. */
 	struct FSlateOnPaintState : public FNoncopyable
 	{
-		FSlateOnPaintState( const FGeometry& InAllottedGeometry, const FSlateRect& InMyClippingRect, FSlateWindowElementList& InOutDrawElements, int32& InLayerId, const FWidgetStyle& InWidgetStyle, ESlateDrawEffect::Type InDrawEffects )
+		FSlateOnPaintState( const FGeometry& InAllottedGeometry, const FSlateRect& InMyClippingRect, FSlateWindowElementList& InOutDrawElements, int32& InLayerId, const FWidgetStyle& InWidgetStyle, ESlateDrawEffect InDrawEffects )
 			: AllottedGeometry( InAllottedGeometry )
 			, MyClippingRect( InMyClippingRect )
 			, WidgetStyle( InWidgetStyle )
@@ -175,7 +175,7 @@ class SProfilerMiniView : public SCompoundWidget
 
 #if	0
 		/** Set operation. */
-		void Set( const FGeometry& InAllottedGeometry, const FSlateRect& InMyClippingRect, FSlateWindowElementList& InOutDrawElements, int32& InLayerId, const FWidgetStyle& InWidgetStyle, ESlateDrawEffect::Type InDrawEffects )
+		void Set( const FGeometry& InAllottedGeometry, const FSlateRect& InMyClippingRect, FSlateWindowElementList& InOutDrawElements, int32& InLayerId, const FWidgetStyle& InWidgetStyle, ESlateDrawEffect InDrawEffects )
 		{
 			_AllottedGeometry = &InAllottedGeometry;
 			_MyClippingRect = &InMyClippingRect;
@@ -193,7 +193,7 @@ class SProfilerMiniView : public SCompoundWidget
 		 
 		FSlateWindowElementList& OutDrawElements;
 		int32& LayerId;
-		const ESlateDrawEffect::Type DrawEffects;
+		const ESlateDrawEffect DrawEffects;
 	};
 
 public:
diff --git a/Engine/Source/Developer/Profiler/Private/Widgets/SProfilerThreadView.cpp b/Engine/Source/Developer/Profiler/Private/Widgets/SProfilerThreadView.cpp
--- a/Engine/Source/Developer/Profiler/Private/Widgets/SProfilerThreadView.cpp
+++ b/Engine/Source/Developer/Profiler/Private/Widgets/SProfilerThreadView.cpp
@@ -82,7 +82,7 @@ int32 SProfilerThreadView::OnPaint( const FPaintArgs& Args, const FGeometry& All
 
 	// Rendering info.
 	const bool bEnabled = ShouldBeEnabled( bParentEnabled );
-	const ESlateDrawEffect::Type DrawEffects = bEnabled ? ESlateDrawEffect::None : ESlateDrawEffect::DisabledEffect;
+	const ESlateDrawEffect DrawEffects = bEnabled ? ESlateDrawEffect::None : ESlateDrawEffect::DisabledEffect;
 	const FSlateBrush* BackgroundBrush = FEditorStyle::GetBrush( "Profiler.LineGraphArea" );
 	const FSlateBrush* WhiteBrush = FEditorStyle::GetBrush( "WhiteTexture" );
 
diff --git a/Engine/Source/Developer/Profiler/Private/Widgets/SProfilerThreadView.h b/Engine/Source/Developer/Profiler/Private/Widgets/SProfilerThreadView.h
--- a/Engine/Source/Developer/Profiler/Private/Widgets/SProfilerThreadView.h
+++ b/Engine/Source/Developer/Profiler/Private/Widgets/SProfilerThreadView.h
@@ -81,7 +81,7 @@ class SProfilerThreadView
 	/** Holds current state provided by OnPaint function, used to simplify drawing. */
 	struct FSlateOnPaintState : public FNoncopyable
 	{
-		FSlateOnPaintState( const FGeometry& InAllottedGeometry, const FSlateRect& InMyClippingRect, FSlateWindowElementList& InOutDrawElements, int32& InLayerId, const FWidgetStyle& InWidgetStyle, ESlateDrawEffect::Type InDrawEffects )
+		FSlateOnPaintState( const FGeometry& InAllottedGeometry, const FSlateRect& InMyClippingRect, FSlateWindowElementList& InOutDrawElements, int32& InLayerId, const FWidgetStyle& InWidgetStyle, ESlateDrawEffect InDrawEffects )
 			: AllottedGeometry( InAllottedGeometry )
 			, AbsoluteClippingRect( InMyClippingRect )
 			, LocalClippingRect( FVector2D::ZeroVector, InAllottedGeometry.Size )
@@ -107,7 +107,7 @@ class SProfilerThreadView
 		 
 		FSlateWindowElementList& OutDrawElements;
 		int32& LayerId;
-		const ESlateDrawEffect::Type DrawEffects;
+		const ESlateDrawEffect DrawEffects;
 
 		const TSharedRef< FSlateFontMeasure > FontMeasureService;
 
diff --git a/Engine/Source/Developer/ScreenShotComparison/Private/Models/ScreenComparisonModel.cpp b/Engine/Source/Developer/ScreenShotComparison/Private/Models/ScreenComparisonModel.cpp
new file mode 100644
--- /dev/null
+++ b/Engine/Source/Developer/ScreenShotComparison/Private/Models/ScreenComparisonModel.cpp
@@ -0,0 +1,226 @@
+// Copyright 1998-2017 Epic Games, Inc. All Rights Reserved.
+
+#include "ScreenComparisonModel.h"
+#include "ISourceControlModule.h"
+#include "ISourceControlOperation.h"
+#include "SourceControlOperations.h"
+#include "ISourceControlProvider.h"
+#include "Paths.h"
+
+DEFINE_LOG_CATEGORY_STATIC(LogScreenshotComparison, Log, All);
+
+FScreenComparisonModel::FScreenComparisonModel(const FComparisonReport& InReport)
+	: Report(InReport)
+	, bComplete(false)
+{
+	const FImageComparisonResult& ComparisonResult = Report.Comparison;
+
+	FString IncomingImage = Report.ReportFolder / ComparisonResult.ReportIncomingFile;
+	FString IncomingMetadata = FPaths::ChangeExtension(IncomingImage, TEXT("json"));
+
+	FileImports.Add(FFileMapping(IncomingImage, ComparisonResult.IncomingFile));
+	FileImports.Add(FFileMapping(IncomingMetadata, FPaths::ChangeExtension(ComparisonResult.IncomingFile, TEXT("json"))));
+}
+
+bool FScreenComparisonModel::IsComplete() const
+{
+	return bComplete;
+}
+
+void FScreenComparisonModel::Complete()
+{
+	FString RelativeReportFolder = Report.ReportFolder;
+	if ( FPaths::MakePathRelativeTo(RelativeReportFolder, *Report.ReportRootDirectory) )
+	{
+		for (;;)
+		{
+			FString ParentFolder = FPaths::GetPath(RelativeReportFolder);
+			if ( ParentFolder.IsEmpty() )
+			{
+				break;
+			}
+			RelativeReportFolder = ParentFolder;
+		}
+
+		FString ReportTopFolder = Report.ReportRootDirectory / RelativeReportFolder;
+		if ( IFileManager::Get().DeleteDirectory(*ReportTopFolder, false, true) )
+		{
+			bComplete = true;
+			OnComplete.Broadcast();
+		}
+	}
+}
+
+TOptional<FAutomationScreenshotMetadata> FScreenComparisonModel::GetMetadata()
+{
+	// Load it.
+	if ( !Metadata.IsSet() )
+	{
+		FString IncomingImage = Report.ReportFolder / Report.Comparison.ReportIncomingFile;
+		FString IncomingMetadata = FPaths::ChangeExtension(IncomingImage, TEXT("json"));
+
+		if ( !IncomingMetadata.IsEmpty() )
+		{
+			FString Json;
+			if ( FFileHelper::LoadFileToString(Json, *IncomingMetadata) )
+			{
+				FAutomationScreenshotMetadata LoadedMetadata;
+				if ( FJsonObjectConverter::JsonObjectStringToUStruct(Json, &LoadedMetadata, 0, 0) )
+				{
+					Metadata = LoadedMetadata;
+				}
+			}
+		}
+	}
+
+	return Metadata;
+}
+
+bool FScreenComparisonModel::AddNew(IScreenShotManagerPtr ScreenshotManager)
+{
+	// Copy the files from the reports location to the destination location
+	TArray<FString> SourceControlFiles;
+	for ( const FFileMapping& Import : FileImports )
+	{
+		FString DestFilePath = ScreenshotManager->GetLocalApprovedFolder() / Import.DestinationFile;
+		IFileManager::Get().Copy(*DestFilePath, *Import.SourceFile, true, true);
+
+		SourceControlFiles.Add(DestFilePath);
+	}
+
+	// Add the files to source control
+	ISourceControlProvider& SourceControlProvider = ISourceControlModule::Get().GetProvider();
+	if ( SourceControlProvider.Execute(ISourceControlOperation::Create<FMarkForAdd>(), SourceControlFiles) == ECommandResult::Failed )
+	{
+		// TODO Error
+	}
+
+	Complete();
+
+	return true;
+}
+
+bool FScreenComparisonModel::Replace(IScreenShotManagerPtr ScreenshotManager)
+{
+	// Delete all the existing files in this area
+	RemoveExistingApproved(ScreenshotManager);
+
+	// Copy files to the approved
+	const FString& LocalApprovedFolder = ScreenshotManager->GetLocalApprovedFolder();
+	const FString ImportIncomingRoot = Report.ReportFolder;
+
+	TArray<FString> SourceControlFiles;
+
+	for ( const FFileMapping& Import : FileImports )
+	{
+		FString DestFilePath = LocalApprovedFolder / Import.DestinationFile;
+		SourceControlFiles.Add(DestFilePath);
+	}
+
+	ISourceControlProvider& SourceControlProvider = ISourceControlModule::Get().GetProvider();
+	if ( SourceControlProvider.Execute(ISourceControlOperation::Create<FRevert>(), SourceControlFiles) == ECommandResult::Failed )
+	{
+		//TODO Error
+	}
+
+	for ( const FFileMapping& Import : FileImports )
+	{
+		FString DestFilePath = LocalApprovedFolder / Import.DestinationFile;
+		IFileManager::Get().Copy(*DestFilePath, *Import.SourceFile, true, true);
+
+		SourceControlFiles.Add(DestFilePath);
+	}
+
+	if ( SourceControlProvider.Execute(ISourceControlOperation::Create<FMarkForAdd>(), SourceControlFiles) == ECommandResult::Failed )
+	{
+		//TODO Error
+	}
+	if ( SourceControlProvider.Execute(ISourceControlOperation::Create<FCheckOut>(), SourceControlFiles) == ECommandResult::Failed )
+	{
+		//TODO Error
+	}
+
+	Complete();
+
+	return true;
+}
+
+bool FScreenComparisonModel::RemoveExistingApproved(IScreenShotManagerPtr ScreenshotManager)
+{
+	TArray<FString> FilesToRemove;
+
+	FString PlatformFolder = Report.ReportFolder;
+	IFileManager::Get().FindFilesRecursive(FilesToRemove, *PlatformFolder, TEXT("*.*"), true, false);
+
+	// Copy files to the approved
+	const FString& LocalApprovedFolder = ScreenshotManager->GetLocalApprovedFolder();
+	const FString ImportApprovedRoot = Report.ReportFolder / TEXT("");
+
+	TArray<FString> SourceControlFiles;
+
+	for ( const FString& File : FilesToRemove )
+	{
+		FString RelativeFile = File;
+		FPaths::MakePathRelativeTo(RelativeFile, *ImportApprovedRoot);
+
+		FString DestFilePath = LocalApprovedFolder / RelativeFile;
+		SourceControlFiles.Add(DestFilePath);
+
+		IFileManager::Get().Delete(*DestFilePath, false, true, false);
+	}
+
+	ISourceControlProvider& SourceControlProvider = ISourceControlModule::Get().GetProvider();
+	if ( SourceControlProvider.Execute(ISourceControlOperation::Create<FDelete>(), SourceControlFiles) == ECommandResult::Failed )
+	{
+		//TODO Error
+	}
+
+	return true;
+}
+
+bool FScreenComparisonModel::AddAlternative(IScreenShotManagerPtr ScreenshotManager)
+{
+	// Copy files to the approved
+	const FString& LocalApprovedFolder = ScreenshotManager->GetLocalApprovedFolder();
+	const FString ImportIncomingRoot = Report.ReportFolder;
+
+	TArray<FString> SourceControlFiles;
+
+	for ( const FFileMapping& Import : FileImports )
+	{
+		FString DestFilePath = LocalApprovedFolder / Import.DestinationFile;
+		SourceControlFiles.Add(DestFilePath);
+	}
+
+	ISourceControlProvider& SourceControlProvider = ISourceControlModule::Get().GetProvider();
+	if ( SourceControlProvider.Execute(ISourceControlOperation::Create<FRevert>(), SourceControlFiles) == ECommandResult::Failed )
+	{
+		//TODO Error
+	}
+
+	for ( const FFileMapping& Import : FileImports )
+	{
+		FString DestFilePath = LocalApprovedFolder / Import.DestinationFile;
+		if ( IFileManager::Get().Copy(*DestFilePath, *Import.SourceFile, false, true) == COPY_OK )
+		{
+			SourceControlFiles.Add(DestFilePath);
+		}
+		else
+		{
+			// TODO Error
+		}
+	}
+
+	if ( SourceControlProvider.Execute(ISourceControlOperation::Create<FMarkForAdd>(), SourceControlFiles) == ECommandResult::Failed )
+	{
+		//TODO Error
+	}
+	if ( SourceControlProvider.Execute(ISourceControlOperation::Create<FCheckOut>(), SourceControlFiles) == ECommandResult::Failed )
+	{
+		//TODO Error
+	}
+
+	Complete();
+
+	return true;
+}
\ No newline at end of file
diff --git a/Engine/Source/Developer/ScreenShotComparison/Private/Models/ScreenComparisonModel.h b/Engine/Source/Developer/ScreenShotComparison/Private/Models/ScreenComparisonModel.h
--- a/Engine/Source/Developer/ScreenShotComparison/Private/Models/ScreenComparisonModel.h
+++ b/Engine/Source/Developer/ScreenShotComparison/Private/Models/ScreenComparisonModel.h
@@ -4,31 +4,51 @@
 #include "CoreMinimal.h"
 #include "ImageComparer.h"
 #include "AutomationWorkerMessages.h"
-
-enum class EComparisonResultType
-{
-	Added,
-	Missing,
-	Compared
-};
+#include "Misc/FileHelper.h"
+#include "JsonObjectConverter.h"
+#include "AutomationWorkerMessages.h"
+#include "Interfaces/IScreenShotManager.h"
 
 class FScreenComparisonModel
 {
 public:
-	FScreenComparisonModel(EComparisonResultType InType)
-		: Type(InType)
-	{
-	}
+	FScreenComparisonModel(const FComparisonReport& InReport);
 
-	EComparisonResultType GetType() const { return Type; }
+	DECLARE_MULTICAST_DELEGATE(FOnComplete);
+	FOnComplete OnComplete;
 
-	FString Folder;
+	FComparisonReport Report;
 
-	TOptional<FImageComparisonResult> ComparisonResult;
+	bool IsComplete() const;
 
-	TOptional<FAutomationScreenshotMetadata> Metadata;
+	void Complete();
+
+	bool AddNew(IScreenShotManagerPtr ScreenshotManager);
+	bool Replace(IScreenShotManagerPtr ScreenshotManager);
+	bool AddAlternative(IScreenShotManagerPtr ScreenshotManager);
+
+	TOptional<FAutomationScreenshotMetadata> GetMetadata();
+
+private:
+	bool RemoveExistingApproved(IScreenShotManagerPtr ScreenshotManager);
 
 private:
+	bool bComplete;
 
-	EComparisonResultType Type;
+	TOptional<FAutomationScreenshotMetadata> Metadata;
+
+	struct FFileMapping
+	{
+		FFileMapping(const FString& InSourceFile, const FString& InDestinationFile)
+			: SourceFile(InSourceFile)
+			, DestinationFile(InDestinationFile)
+		{
+		}
+
+		FString SourceFile;
+		FString DestinationFile;
+	};
+
+	// 
+	TArray<FFileMapping> FileImports;
 };
