Automation - The automation report is now written out with the UTF-8 BOM, to ensure correct interpretation of the browser.  Further tweaks to the format, fixing several issues with the Lighting tests.

diff --git a/Engine/Plugins/Tests/EditorTests/Source/EditorTests/Private/UnrealEd/LightingTests.cpp b/Engine/Plugins/Tests/EditorTests/Source/EditorTests/Private/UnrealEd/LightingTests.cpp
--- a/Engine/Plugins/Tests/EditorTests/Source/EditorTests/Private/UnrealEd/LightingTests.cpp
+++ b/Engine/Plugins/Tests/EditorTests/Source/EditorTests/Private/UnrealEd/LightingTests.cpp
@@ -9,6 +9,7 @@
 #include "Components/PointLightComponent.h"
 #include "Engine/PointLight.h"
 #include "Editor.h"
+#include "EngineUtils.h"
 
 // Automation
 #include "Tests/AutomationEditorCommon.h"
@@ -150,9 +151,9 @@ bool FLightingPromotionPointLightPlaceRotScaleTest::RunTest(const FString& Param
 		bool RotationsAreEqual = CurrentRotation.Equals(POINT_LIGHT_UPDATED_ROTATION, 1);
 
 		TestTrue(TEXT("The placed point light was not found."), LightingTestHelpers::DoesActorExistInTheLevel(CurrentLevel, PointLight->GetName(), PointLight->GetClass()));
-		TestEqual<FVector>(TEXT("The point light is not in correct location"), POINT_LIGHT_UPDATED_LOCATION, CurrentLocation);
+		TestEqual(TEXT("The point light is not in correct location"), CurrentLocation, POINT_LIGHT_UPDATED_LOCATION);
 		TestTrue(TEXT("The point light is not rotated correctly."), RotationsAreEqual);
-		TestEqual<FVector>(TEXT("The point light is not scaled correctly."), POINT_LIGHT_UPDATED_SCALE3D, CurrentScale3D);
+		TestEqual(TEXT("The point light is not scaled correctly."), CurrentScale3D, POINT_LIGHT_UPDATED_SCALE3D);
 
 		return true;
 	}
@@ -185,9 +186,9 @@ bool FLightingPromotionModifyProperties::RunTest(const FString& Parameters)
 		LightingTestHelpers::SetPropertyByName(PointLight->PointLightComponent, TEXT("AttenuationRadius"), TEXT("1024.f"));
 
 		//** VERIFY **//
-		TestEqual<float>(TEXT("Light brightness property was not modified."), 1000.f, PointLight->PointLightComponent->Intensity);
-		TestEqual<FColor>(TEXT("Light color property was not modified."), FColor(0,0,255), PointLight->PointLightComponent->LightColor);
-		TestEqual<float>(TEXT("Light attenuation radius was not modified."), 1024.f, PointLight->PointLightComponent->AttenuationRadius);
+		TestEqual(TEXT("Light Brightness"), PointLight->PointLightComponent->Intensity, 1000.f);
+		TestEqual(TEXT("Light Color"), PointLight->PointLightComponent->LightColor, FColor(0, 0, 255));
+		TestEqual(TEXT("Light Attenuation Radius"), PointLight->PointLightComponent->AttenuationRadius, 1024.f);
 
 		return true;
 	}
@@ -228,13 +229,13 @@ bool FLightingPromotionDuplicationTest::RunTest(const FString& Parameters)
 		//** Verify **//
 		int32 NumberOfPointLights = 0;
 		// Count the number of point lights in the level.
-		for (TObjectIterator<APointLight> It; It; ++It)
+		for ( TActorIterator<APointLight> It(World); It; ++It )
 		{
 			NumberOfPointLights++;
 		}
 		
 		// We are expecting three point lights to be in the level now.
-		TestEqual<int32>(TEXT("After duplicating a light the total number of them in the level is not correct."), 3, NumberOfPointLights);
+		TestEqual(TEXT("The light count after duplication"), NumberOfPointLights, 3);
 
 		return true;
 	}
@@ -306,7 +307,7 @@ bool FLightPointLightSetLocation::RunTest(const FString& Parameters)
 		FVector CurrentLocation;
 		LightingTestHelpers::GetActorCurrentLocation(CurrentLevel, PointLight->GetName(), CurrentLocation);
 
-		TestEqual<FVector>(TEXT("The point light is not in correct location"), POINT_LIGHT_UPDATED_LOCATION, CurrentLocation);
+		TestEqual(TEXT("The point light is not in correct location"), CurrentLocation, POINT_LIGHT_UPDATED_LOCATION);
 		return true;
 	}
 
diff --git a/Engine/Source/Developer/AutomationController/Private/AutomationControllerManger.cpp b/Engine/Source/Developer/AutomationController/Private/AutomationControllerManger.cpp
--- a/Engine/Source/Developer/AutomationController/Private/AutomationControllerManger.cpp
+++ b/Engine/Source/Developer/AutomationController/Private/AutomationControllerManger.cpp
@@ -476,7 +476,7 @@ void FAutomationControllerManager::GenerateHtmlTestPassSummary(FDateTime Timesta
 		FString Html = FString::Format(*MasterTemplate, Args);
 
 		FString ReportFileName = FString::Printf(TEXT("%s/AutomationReport-%d-%s.html"), *ReportOutputFolder, FEngineVersion::Current().GetChangelist(), *Timestamp.ToString());
-		if ( !FFileHelper::SaveStringToFile(Html, *ReportFileName) )
+		if ( !FFileHelper::SaveStringToFile(Html, *ReportFileName, FFileHelper::EEncodingOptions::ForceUTF8) )
 		{
 			GLog->Logf(ELogVerbosity::Error, TEXT("Test Report Html is invalid - report not generated."));
 		}
diff --git a/Engine/Source/Runtime/Core/Public/Misc/AutomationTest.h b/Engine/Source/Runtime/Core/Public/Misc/AutomationTest.h
--- a/Engine/Source/Runtime/Core/Public/Misc/AutomationTest.h
+++ b/Engine/Source/Runtime/Core/Public/Misc/AutomationTest.h
@@ -1063,6 +1063,22 @@ public:
 		}
 	}
 
+	void TestEqual(const FString& What, const FVector Actual, const FVector Expected, float Tolerance = 1.e-4)
+	{
+		if ( !Expected.Equals(Actual, Tolerance) )
+		{
+			AddError(FString::Printf(TEXT("Expected '%s' to be %s, but it was %s within tolerance %f."), *What, *Expected.ToString(), *Actual.ToString(), Tolerance), 1);
+		}
+	}
+
+	void TestEqual(const FString& What, const FColor Actual, const FColor Expected)
+	{
+		if ( Expected != Actual )
+		{
+			AddError(FString::Printf(TEXT("Expected '%s' to be %s, but it was %s."), *What, *Expected.ToString(), *Actual.ToString()), 1);
+		}
+	}
+
 	/**
 	 * Logs an error if the two values are not equal.
 	 *
