Added more checks to the Light Duplicate and Copy Paste automated test to help narrow down the failure we see in CIS.

diff --git a/Engine/Plugins/Tests/EditorTests/Source/EditorTests/Private/UnrealEd/LightingTests.cpp b/Engine/Plugins/Tests/EditorTests/Source/EditorTests/Private/UnrealEd/LightingTests.cpp
--- a/Engine/Plugins/Tests/EditorTests/Source/EditorTests/Private/UnrealEd/LightingTests.cpp
+++ b/Engine/Plugins/Tests/EditorTests/Source/EditorTests/Private/UnrealEd/LightingTests.cpp
@@ -214,18 +214,42 @@ bool FLightingPromotionDuplicationTest::RunTest(const FString& Parameters)
 		//** TEST **//
 		// Add a point light to the level.
 		APointLight* PointLight = Cast<APointLight>(GEditor->AddActor(World->GetCurrentLevel(), APointLight::StaticClass(), FTransform()));
+
+		//** Verify **//
+		int32 NumberOfPointLights = 0;
+		// Count the number of point lights in the level.
+		for ( TActorIterator<APointLight> It(World); It; ++It )
+		{
+			NumberOfPointLights++;
+		}
+
+		// Make sure there's only one point light in the level.
+		TestEqual(TEXT("The light count before copy/paste"), NumberOfPointLights, 1);
+
 		// Deselect all and then Select the light
 		LightingTestHelpers::SelectActorInLevel(PointLight);
 		// Copy and Paste.
 		GEngine->Exec(World, TEXT("EDIT COPY"));
 		GEngine->Exec(World, TEXT("EDIT PASTE"));
+
+		//** Verify **//
+		NumberOfPointLights = 0;
+		// Count the number of point lights in the level.
+		for ( TActorIterator<APointLight> It(World); It; ++It )
+		{
+			NumberOfPointLights++;
+		}
+
+		// We are expecting two point lights to be in the level now.
+		TestEqual(TEXT("The light count after copy/paste"), NumberOfPointLights, 2);
+
 		// Deselect all and then select a light
 		LightingTestHelpers::SelectActorInLevel(PointLight);
 		// Duplicate the light
 		GEngine->Exec(World, TEXT("DUPLICATE"));
 
 		//** Verify **//
-		int32 NumberOfPointLights = 0;
+		NumberOfPointLights = 0;
 		// Count the number of point lights in the level.
 		for ( TActorIterator<APointLight> It(World); It; ++It )
 		{
