Attempts ot track down and fix mac duplicate actors automation test

diff --git a/Engine/Plugins/Tests/EditorTests/Source/EditorTests/Private/UnrealEd/LightingTests.cpp b/Engine/Plugins/Tests/EditorTests/Source/EditorTests/Private/UnrealEd/LightingTests.cpp
--- a/Engine/Plugins/Tests/EditorTests/Source/EditorTests/Private/UnrealEd/LightingTests.cpp
+++ b/Engine/Plugins/Tests/EditorTests/Source/EditorTests/Private/UnrealEd/LightingTests.cpp
@@ -16,6 +16,8 @@
 
 // Misc
 #include "ScopedTransaction.h"
+#include "Engine/Selection.h"
+#include "EditorModeManager.h"
 
 #define POINTLIGHT_TRANSFORM			FTransform()
 #define POINT_LIGHT_UPDATED_LOCATION	FVector(50.f, 50.f, 50.f)
@@ -209,8 +211,11 @@ bool FLightingPromotionDuplicationTest::RunTest(const FString& Parameters)
 	// Test Summary
 	AddInfo(TEXT("Duplicate and Copy Paste\n- Duplicates a point light.\n- Copies and Pastes a point light."));
 
+	GLevelEditorModeTools().ActivateDefaultMode();
+
 	if (!LightingTestHelpers::DoesActorExistInTheLevel(CurrentLevel, TEXT("PointLight"), APointLight::StaticClass()))
 	{
+	
 		//** TEST **//
 		// Add a point light to the level.
 		APointLight* PointLight = Cast<APointLight>(GEditor->AddActor(World->GetCurrentLevel(), APointLight::StaticClass(), FTransform()));
@@ -228,9 +233,16 @@ bool FLightingPromotionDuplicationTest::RunTest(const FString& Parameters)
 
 		// Deselect all and then Select the light
 		LightingTestHelpers::SelectActorInLevel(PointLight);
+
+		TestEqual(TEXT("Selected actor count"), GEditor->GetSelectedActorCount(), 1);
+
+		TestEqual<APointLight*>(TEXT("PointLight is selected"), PointLight, GEditor->GetSelectedActors()->GetTop<APointLight>());
+
 		// Copy and Paste.
-		GEngine->Exec(World, TEXT("EDIT COPY"));
-		GEngine->Exec(World, TEXT("EDIT PASTE"));
+		GEditor->CopySelectedActorsToClipboard(World, false);
+
+		FText TransDescription = FText::FromString("Automation Test Copy/Paste");
+		GEditor->PasteSelectedActorsFromClipboard(World, TransDescription, EPasteTo::PT_OriginalLocation);
 
 		//** Verify **//
 		NumberOfPointLights = 0;
@@ -245,6 +257,7 @@ bool FLightingPromotionDuplicationTest::RunTest(const FString& Parameters)
 
 		// Deselect all and then select a light
 		LightingTestHelpers::SelectActorInLevel(PointLight);
+
 		// Duplicate the light
 		GEngine->Exec(World, TEXT("DUPLICATE"));
 
