Disable copy/paste lights automation test.  It isnt finding a legit error and is failing in unreproducible ways

diff --git a/Engine/Plugins/Tests/EditorTests/Source/EditorTests/Private/UnrealEd/LightingTests.cpp b/Engine/Plugins/Tests/EditorTests/Source/EditorTests/Private/UnrealEd/LightingTests.cpp
--- a/Engine/Plugins/Tests/EditorTests/Source/EditorTests/Private/UnrealEd/LightingTests.cpp
+++ b/Engine/Plugins/Tests/EditorTests/Source/EditorTests/Private/UnrealEd/LightingTests.cpp
@@ -198,14 +198,14 @@ bool FLightingPromotionModifyProperties::RunTest(const FString& Parameters)
 }
 
 
-/**
+/*
+/ **
 * Lighting Promotion Test - Duplicate/Copy Paste a point light.
-*/
+* /
 IMPLEMENT_SIMPLE_AUTOMATION_TEST(FLightingPromotionDuplicationTest, "System.Promotion.Editor.Lighting.Duplicate and Copy Paste", EAutomationTestFlags::EditorContext | EAutomationTestFlags::EngineFilter)
 
 bool FLightingPromotionDuplicationTest::RunTest(const FString& Parameters)
 {
-	//** SETUP **//
 	UWorld* World = FAutomationEditorCommonUtils::CreateNewMap();
 	ULevel* CurrentLevel = World->GetCurrentLevel();
 	// Test Summary
@@ -215,12 +215,9 @@ bool FLightingPromotionDuplicationTest::RunTest(const FString& Parameters)
 
 	if (!LightingTestHelpers::DoesActorExistInTheLevel(CurrentLevel, TEXT("PointLight"), APointLight::StaticClass()))
 	{
-	
-		//** TEST **//
 		// Add a point light to the level.
 		APointLight* PointLight = Cast<APointLight>(GEditor->AddActor(World->GetCurrentLevel(), APointLight::StaticClass(), FTransform()));
 
-		//** Verify **//
 		int32 NumberOfPointLights = 0;
 		// Count the number of point lights in the level.
 		for ( TActorIterator<APointLight> It(World); It; ++It )
@@ -244,7 +241,6 @@ bool FLightingPromotionDuplicationTest::RunTest(const FString& Parameters)
 		FText TransDescription = FText::FromString("Automation Test Copy/Paste");
 		GEditor->PasteSelectedActorsFromClipboard(World, TransDescription, EPasteTo::PT_OriginalLocation);
 
-		//** Verify **//
 		NumberOfPointLights = 0;
 		// Count the number of point lights in the level.
 		for ( TActorIterator<APointLight> It(World); It; ++It )
@@ -261,7 +257,6 @@ bool FLightingPromotionDuplicationTest::RunTest(const FString& Parameters)
 		// Duplicate the light
 		GEngine->Exec(World, TEXT("DUPLICATE"));
 
-		//** Verify **//
 		NumberOfPointLights = 0;
 		// Count the number of point lights in the level.
 		for ( TActorIterator<APointLight> It(World); It; ++It )
@@ -277,7 +272,7 @@ bool FLightingPromotionDuplicationTest::RunTest(const FString& Parameters)
 
 	AddError(TEXT("A point light already exists in this level which would dirty the test results."));
 	return false;
-}
+}*/
 
 
 //////////////////////////////////////////////////////////////////////////
