forgot a file with the last submit

diff --git a/Engine/Plugins/Tests/EditorTests/Source/EditorTests/EditorTests.Build.cs b/Engine/Plugins/Tests/EditorTests/Source/EditorTests/EditorTests.Build.cs
--- a/Engine/Plugins/Tests/EditorTests/Source/EditorTests/EditorTests.Build.cs
+++ b/Engine/Plugins/Tests/EditorTests/Source/EditorTests/EditorTests.Build.cs
@@ -62,7 +62,10 @@ public class EditorTests : ModuleRules
 				"AudioEditor",
 				"AnimGraphRuntime",
 				"MeshMergeUtilities",
-				"MaterialBaking"
+				"MaterialBaking",
+                "MeshDescription",
+                "MeshBuilder",
+                "RawMesh",
 			}
 		);
 
