- Make sure meshdescription ddc key is trustable. Sparse array serialization is not bit equal for similar data
- Make FStaticMeshSourceModel::Meshdescription transient
- Meshdescription ddc now contain only the meshdescription and use the FRawMesh as the key
- Automation tests now do the conversion in both ways MD->RM->MD and RM->MD->RM
- Automation tests code get cleanup and refactor
- resave the meshdescription automation tests assets

diff --git a/Engine/Plugins/Tests/EditorTests/Source/EditorTests/Private/UnrealEd/MeshDescriptionTests.cpp b/Engine/Plugins/Tests/EditorTests/Source/EditorTests/Private/UnrealEd/MeshDescriptionTests.cpp
--- a/Engine/Plugins/Tests/EditorTests/Source/EditorTests/Private/UnrealEd/MeshDescriptionTests.cpp
+++ b/Engine/Plugins/Tests/EditorTests/Source/EditorTests/Private/UnrealEd/MeshDescriptionTests.cpp
@@ -38,7 +38,8 @@ public:
 	bool Execute(FAutomationTestExecutionInfo& ExecutionInfo);
 	bool ConversionTest(FAutomationTestExecutionInfo& ExecutionInfo);
 private:
-	bool CompareMeshDescription(const FString& AssetName, FAutomationTestExecutionInfo& ExecutionInfo, const UMeshDescription* ReferenceMeshDescription, const UMeshDescription* MeshDescription);
+	bool CompareRawMesh(const FString& AssetName, FAutomationTestExecutionInfo& ExecutionInfo, const FRawMesh& ReferenceRawMesh, const FRawMesh& ResultRawMesh) const;
+	bool CompareMeshDescription(const FString& AssetName, FAutomationTestExecutionInfo& ExecutionInfo, const UMeshDescription* ReferenceMeshDescription, const UMeshDescription* MeshDescription) const;
 };
 
 class FMeshDescriptionTests
@@ -127,170 +128,245 @@ bool FMeshDescriptionTest::Execute(FAutomationTestExecutionInfo& ExecutionInfo)
 	return bSuccess;
 }
 
-bool FMeshDescriptionTest::CompareMeshDescription(const FString& AssetName, FAutomationTestExecutionInfo& ExecutionInfo, const UMeshDescription* ReferenceMeshDescription, const UMeshDescription* MeshDescription)
+template<typename T>
+void StructureArrayCompare(const FString& ConversionName, const FString& AssetName, FAutomationTestExecutionInfo& ExecutionInfo, bool& bIsSame, const FString& VectorArrayName, const TArray<T>& ReferenceArray, const TArray<T>& ResultArray)
 {
-	//////////////////////////////////////////////////////////////////////////
-	//Gather the reference data
-	const TVertexAttributeArray<FVector>& ReferenceVertexPositions = ReferenceMeshDescription->VertexAttributes().GetAttributes<FVector>(MeshAttribute::Vertex::Position);
-
-	const TVertexInstanceAttributeArray<FVector>& ReferenceVertexInstanceNormals = ReferenceMeshDescription->VertexInstanceAttributes().GetAttributes<FVector>(MeshAttribute::VertexInstance::Normal);
-	const TVertexInstanceAttributeArray<FVector>& ReferenceVertexInstanceTangents = ReferenceMeshDescription->VertexInstanceAttributes().GetAttributes<FVector>(MeshAttribute::VertexInstance::Tangent);
-	const TVertexInstanceAttributeArray<float>& ReferenceVertexInstanceBinormalSigns = ReferenceMeshDescription->VertexInstanceAttributes().GetAttributes<float>(MeshAttribute::VertexInstance::BinormalSign);
-	const TVertexInstanceAttributeArray<FVector4>& ReferenceVertexInstanceColors = ReferenceMeshDescription->VertexInstanceAttributes().GetAttributes<FVector4>(MeshAttribute::VertexInstance::Color);
-	const TVertexInstanceAttributeIndicesArray<FVector2D>& ReferenceVertexInstanceUVs = ReferenceMeshDescription->VertexInstanceAttributes().GetAttributesSet<FVector2D>(MeshAttribute::VertexInstance::TextureCoordinate);
-
-	const TEdgeAttributeArray<bool>& ReferenceEdgeHardnesses = ReferenceMeshDescription->EdgeAttributes().GetAttributes<bool>(MeshAttribute::Edge::IsHard);
-
-	const TPolygonGroupAttributeArray<int>& ReferencePolygonGroupMaterialIndex = ReferenceMeshDescription->PolygonGroupAttributes().GetAttributes<int>(MeshAttribute::PolygonGroup::MaterialIndex);
-
-
-	//////////////////////////////////////////////////////////////////////////
-	//Gather the result data
-	const TVertexAttributeArray<FVector>& ResultVertexPositions = MeshDescription->VertexAttributes().GetAttributes<FVector>(MeshAttribute::Vertex::Position);
-
-	const TVertexInstanceAttributeArray<FVector>& ResultVertexInstanceNormals = MeshDescription->VertexInstanceAttributes().GetAttributes<FVector>(MeshAttribute::VertexInstance::Normal);
-	const TVertexInstanceAttributeArray<FVector>& ResultVertexInstanceTangents = MeshDescription->VertexInstanceAttributes().GetAttributes<FVector>(MeshAttribute::VertexInstance::Tangent);
-	const TVertexInstanceAttributeArray<float>& ResultVertexInstanceBinormalSigns = MeshDescription->VertexInstanceAttributes().GetAttributes<float>(MeshAttribute::VertexInstance::BinormalSign);
-	const TVertexInstanceAttributeArray<FVector4>& ResultVertexInstanceColors = MeshDescription->VertexInstanceAttributes().GetAttributes<FVector4>(MeshAttribute::VertexInstance::Color);
-	const TVertexInstanceAttributeIndicesArray<FVector2D>& ResultVertexInstanceUVs = MeshDescription->VertexInstanceAttributes().GetAttributesSet<FVector2D>(MeshAttribute::VertexInstance::TextureCoordinate);
-
-	const TEdgeAttributeArray<bool>& ResultEdgeHardnesses = MeshDescription->EdgeAttributes().GetAttributes<bool>(MeshAttribute::Edge::IsHard);
-
-	const TPolygonGroupAttributeArray<int>& ResultPolygonGroupMaterialIndex = MeshDescription->PolygonGroupAttributes().GetAttributes<int>(MeshAttribute::PolygonGroup::MaterialIndex);
-
-	
-	//////////////////////////////////////////////////////////////////////////
-	// Do the comparison
-	bool bAllSame = true;
-	//Positions
-	if (ReferenceVertexPositions.Num() != ResultVertexPositions.Num())
+	if (ReferenceArray.Num() != ResultArray.Num())
 	{
-		ExecutionInfo.AddEvent(FAutomationEvent(EAutomationEventType::Error, FString::Printf(TEXT("The %s conversion to RawMesh is not lossless, vertex positions count is different. Vertex count expected [%d] result [%d]"),
+		ExecutionInfo.AddEvent(FAutomationEvent(EAutomationEventType::Error, FString::Printf(TEXT("The %s conversion %s is not lossless, %s count is different. %s count expected [%d] result [%d]"),
 			*AssetName,
-			ReferenceVertexPositions.Num(),
-			ResultVertexPositions.Num())));
-		bAllSame = false;
+			*ConversionName,
+			*VectorArrayName,
+			*VectorArrayName,
+			ReferenceArray.Num(),
+			ResultArray.Num())));
+		bIsSame = false;
 	}
 	else
 	{
-		for (FVertexID VertexID : ReferenceMeshDescription->Vertices().GetElementIDs())
+		for (int32 VertexIndex = 0; VertexIndex < ReferenceArray.Num(); ++VertexIndex)
 		{
-			if (ReferenceVertexPositions[VertexID] != ResultVertexPositions[VertexID])
+			if (ReferenceArray[VertexIndex] != ResultArray[VertexIndex])
 			{
-				ExecutionInfo.AddEvent(FAutomationEvent(EAutomationEventType::Error, FString::Printf(TEXT("The %s conversion to RawMesh is not lossless, vertex position array is different. VertexID [%d] expected position [%s] result [%s]"),
+				ExecutionInfo.AddEvent(FAutomationEvent(EAutomationEventType::Error, FString::Printf(TEXT("The %s conversion %s is not lossless, %s array is different. Array index [%d] expected %s [%s] result [%s]"),
 					*AssetName,
-					VertexID.GetValue(),
-					*ReferenceVertexPositions[VertexID].ToString(),
-					*ResultVertexPositions[VertexID].ToString())));
-				bAllSame = false;
+					*ConversionName,
+					*VectorArrayName,
+					VertexIndex,
+					*VectorArrayName,
+					*ReferenceArray[VertexIndex].ToString(),
+					*ResultArray[VertexIndex].ToString())));
+				bIsSame = false;
 				break;
 			}
 		}
 	}
-	//Normals
-	if (ReferenceVertexInstanceNormals.Num() != ResultVertexInstanceNormals.Num())
+}
+
+template<typename T>
+void NumberArrayCompare(const FString& ConversionName, const FString& AssetName, FAutomationTestExecutionInfo& ExecutionInfo, bool& bIsSame, const FString& VectorArrayName, const TArray<T>& ReferenceArray, const TArray<T>& ResultArray)
+{
+	if (ReferenceArray.Num() != ResultArray.Num())
 	{
-		ExecutionInfo.AddEvent(FAutomationEvent(EAutomationEventType::Error, FString::Printf(TEXT("The %s conversion to RawMesh is not lossless, vertex instance normals count is different. Normals count expected [%d] result [%d]"),
+		ExecutionInfo.AddEvent(FAutomationEvent(EAutomationEventType::Error, FString::Printf(TEXT("The %s conversion %s is not lossless, %s count is different. %s count expected [%d] result [%d]"),
 			*AssetName,
-			ReferenceVertexInstanceNormals.Num(),
-			ResultVertexInstanceNormals.Num())));
-		bAllSame = false;
+			*ConversionName,
+			*VectorArrayName,
+			*VectorArrayName,
+			ReferenceArray.Num(),
+			ResultArray.Num())));
+		bIsSame = false;
 	}
 	else
 	{
-		for (FVertexInstanceID VertexInstanceID : ReferenceMeshDescription->VertexInstances().GetElementIDs())
+		for (int32 VertexIndex = 0; VertexIndex < ReferenceArray.Num(); ++VertexIndex)
 		{
-			if (ReferenceVertexInstanceNormals[VertexInstanceID] != ResultVertexInstanceNormals[VertexInstanceID])
+			if (ReferenceArray[VertexIndex] != ResultArray[VertexIndex])
 			{
-				ExecutionInfo.AddEvent(FAutomationEvent(EAutomationEventType::Error, FString::Printf(TEXT("The %s conversion to RawMesh is not lossless, vertex instance normals array is different. VertexInstanceID [%d] expected normal [%s] result [%s]"),
+				ExecutionInfo.AddEvent(FAutomationEvent(EAutomationEventType::Error, FString::Printf(TEXT("The %s conversion %s is not lossless, %s array is different. Array index [%d] expected %s [%d] result [%d]"),
 					*AssetName,
-					VertexInstanceID.GetValue(),
-					*ReferenceVertexInstanceNormals[VertexInstanceID].ToString(),
-					*ResultVertexInstanceNormals[VertexInstanceID].ToString())));
-				bAllSame = false;
+					*ConversionName,
+					*VectorArrayName,
+					VertexIndex,
+					*VectorArrayName,
+					ReferenceArray[VertexIndex],
+					ResultArray[VertexIndex])));
+				bIsSame = false;
 				break;
 			}
 		}
 	}
+}
+
+bool FMeshDescriptionTest::CompareRawMesh(const FString& AssetName, FAutomationTestExecutionInfo& ExecutionInfo, const FRawMesh& ReferenceRawMesh, const FRawMesh& ResultRawMesh) const
+{
+	//////////////////////////////////////////////////////////////////////////
+	// Do the comparison
+	bool bAllSame = true;
+
+	FString ConversionName = TEXT("RawMesh to MeshDescription to RawMesh");
+
+	//Positions
+	StructureArrayCompare<FVector>(ConversionName, AssetName, ExecutionInfo, bAllSame, TEXT("vertex positions"), ReferenceRawMesh.VertexPositions, ResultRawMesh.VertexPositions);
+
+	//Normals
+	StructureArrayCompare<FVector>(ConversionName, AssetName, ExecutionInfo, bAllSame, TEXT("vertex instance normals"), ReferenceRawMesh.WedgeTangentZ, ResultRawMesh.WedgeTangentZ);
+
 	//Tangents
-	if (ReferenceVertexInstanceTangents.Num() != ResultVertexInstanceTangents.Num())
-	{
-		ExecutionInfo.AddEvent(FAutomationEvent(EAutomationEventType::Error, FString::Printf(TEXT("The %s conversion to RawMesh is not lossless, vertex instance Tangents count is different. Tangents count expected [%d] result [%d]"),
-			*AssetName,
-			ReferenceVertexInstanceTangents.Num(),
-			ResultVertexInstanceTangents.Num())));
-		bAllSame = false;
-	}
-	else
+	StructureArrayCompare<FVector>(ConversionName, AssetName, ExecutionInfo, bAllSame, TEXT("vertex instance tangents"), ReferenceRawMesh.WedgeTangentX, ResultRawMesh.WedgeTangentX);
+
+	//BiNormal
+	StructureArrayCompare<FVector>(ConversionName, AssetName, ExecutionInfo, bAllSame, TEXT("vertex instance binormals"), ReferenceRawMesh.WedgeTangentY, ResultRawMesh.WedgeTangentY);
+
+	//Colors
+	StructureArrayCompare<FColor>(ConversionName, AssetName, ExecutionInfo, bAllSame, TEXT("vertex instance colors"), ReferenceRawMesh.WedgeColors, ResultRawMesh.WedgeColors);
+
+	//Uvs
+	for (int32 UVIndex = 0; UVIndex < MAX_MESH_TEXTURE_COORDS; ++UVIndex)
 	{
-		for (FVertexInstanceID VertexInstanceID : ReferenceMeshDescription->VertexInstances().GetElementIDs())
-		{
-			if (ReferenceVertexInstanceTangents[VertexInstanceID] != ResultVertexInstanceTangents[VertexInstanceID])
-			{
-				ExecutionInfo.AddEvent(FAutomationEvent(EAutomationEventType::Error, FString::Printf(TEXT("The %s conversion to RawMesh is not lossless, vertex instance Tangents array is different. VertexInstanceID [%d] expected Tangent [%s] result [%s]"),
-					*AssetName,
-					VertexInstanceID.GetValue(),
-					*ReferenceVertexInstanceTangents[VertexInstanceID].ToString(),
-					*ResultVertexInstanceTangents[VertexInstanceID].ToString())));
-				bAllSame = false;
-				break;
-			}
-		}
+		FString UVIndexName = FString::Printf(TEXT("vertex instance UVs(%d)"), UVIndex);
+		StructureArrayCompare<FVector2D>(ConversionName, AssetName, ExecutionInfo, bAllSame, UVIndexName, ReferenceRawMesh.WedgeTexCoords[UVIndex], ResultRawMesh.WedgeTexCoords[UVIndex]);
 	}
-	//BiNormal signs
-	if (ReferenceVertexInstanceBinormalSigns.Num() != ResultVertexInstanceBinormalSigns.Num())
+	
+	//Indices
+	NumberArrayCompare<uint32>(ConversionName, AssetName, ExecutionInfo, bAllSame, TEXT("vertex positions"), ReferenceRawMesh.WedgeIndices, ResultRawMesh.WedgeIndices);
+
+	//Face
+	NumberArrayCompare<int32>(ConversionName, AssetName, ExecutionInfo, bAllSame, TEXT("face material"), ReferenceRawMesh.FaceMaterialIndices, ResultRawMesh.FaceMaterialIndices);
+
+	//Smoothing Mask
+	NumberArrayCompare<uint32>(ConversionName, AssetName, ExecutionInfo, bAllSame, TEXT("smoothing mask"), ReferenceRawMesh.FaceSmoothingMasks, ResultRawMesh.FaceSmoothingMasks);
+
+	return bAllSame;
+}
+
+template<typename T, typename U, typename V>
+void MeshDescriptionStructureArrayCompare(const FString& ConversionName, const FString& AssetName, FAutomationTestExecutionInfo& ExecutionInfo, bool& bIsSame, const U& ElementIterator, const FString& VectorArrayName, const T& ReferenceArray, const T& ResultArray)
+{
+	if (ReferenceArray.Num() != ResultArray.Num())
 	{
-		ExecutionInfo.AddEvent(FAutomationEvent(EAutomationEventType::Error, FString::Printf(TEXT("The %s conversion to RawMesh is not lossless, vertex instance BinormalSigns count is different. BinormalSigns count expected [%d] result [%d]"),
+		ExecutionInfo.AddEvent(FAutomationEvent(EAutomationEventType::Error, FString::Printf(TEXT("The %s conversion %s is not lossless, %s count is different. %s count expected [%d] result [%d]"),
 			*AssetName,
-			ReferenceVertexInstanceBinormalSigns.Num(),
-			ResultVertexInstanceBinormalSigns.Num())));
-		bAllSame = false;
+			*ConversionName,
+			*VectorArrayName,
+			*VectorArrayName,
+			ReferenceArray.Num(),
+			ResultArray.Num())));
+		bIsSame = false;
 	}
 	else
 	{
-		for (FVertexInstanceID VertexInstanceID : ReferenceMeshDescription->VertexInstances().GetElementIDs())
+		for (V ElementID : ElementIterator.GetElementIDs())
 		{
-			if (ReferenceVertexInstanceBinormalSigns[VertexInstanceID] != ResultVertexInstanceBinormalSigns[VertexInstanceID])
+			if (ReferenceArray[ElementID] != ResultArray[ElementID])
 			{
-				ExecutionInfo.AddEvent(FAutomationEvent(EAutomationEventType::Error, FString::Printf(TEXT("The %s conversion to RawMesh is not lossless, vertex instance BinormalSigns array is different. VertexInstanceID [%d] expected binormal sign [%f] result [%f]"),
+				ExecutionInfo.AddEvent(FAutomationEvent(EAutomationEventType::Error, FString::Printf(TEXT("The %s conversion %s is not lossless, %s array is different. Array index [%d] expected %s [%s] result [%s]"),
 					*AssetName,
-					VertexInstanceID.GetValue(),
-					ReferenceVertexInstanceBinormalSigns[VertexInstanceID],
-					ResultVertexInstanceBinormalSigns[VertexInstanceID])));
-				bAllSame = false;
+					*ConversionName,
+					*VectorArrayName,
+					ElementID.GetValue(),
+					*VectorArrayName,
+					*ReferenceArray[ElementID].ToString(),
+					*ResultArray[ElementID].ToString())));
+				bIsSame = false;
 				break;
 			}
 		}
 	}
-	//Colors
-	if (ReferenceVertexInstanceColors.Num() != ResultVertexInstanceColors.Num())
+}
+
+template<typename T, typename U, typename V>
+void MeshDescriptionNumberArrayCompare(const FString& ConversionName, const FString& AssetName, FAutomationTestExecutionInfo& ExecutionInfo, bool& bIsSame, const U& ElementIterator, const FString& VectorArrayName, const T& ReferenceArray, const T& ResultArray)
+{
+	if (ReferenceArray.Num() != ResultArray.Num())
 	{
-		ExecutionInfo.AddEvent(FAutomationEvent(EAutomationEventType::Error, FString::Printf(TEXT("The %s conversion to RawMesh is not lossless, vertex instance Colors count is different. Colors count expected [%d] result [%d]"),
+		ExecutionInfo.AddEvent(FAutomationEvent(EAutomationEventType::Error, FString::Printf(TEXT("The %s conversion %s is not lossless, %s count is different. %s count expected [%d] result [%d]"),
 			*AssetName,
-			ReferenceVertexInstanceColors.Num(),
-			ResultVertexInstanceColors.Num())));
-		bAllSame = false;
+			*ConversionName,
+			*VectorArrayName,
+			*VectorArrayName,
+			ReferenceArray.Num(),
+			ResultArray.Num())));
+		bIsSame = false;
 	}
 	else
 	{
-		for (FVertexInstanceID VertexInstanceID : ReferenceMeshDescription->VertexInstances().GetElementIDs())
+		for (V ElementID : ElementIterator.GetElementIDs())
 		{
-			if (ReferenceVertexInstanceColors[VertexInstanceID] != ResultVertexInstanceColors[VertexInstanceID])
+			if (ReferenceArray[ElementID] != ResultArray[ElementID])
 			{
-				ExecutionInfo.AddEvent(FAutomationEvent(EAutomationEventType::Error, FString::Printf(TEXT("The %s conversion to RawMesh is not lossless, vertex instance Colors array is different. VertexInstanceID [%d] expected Color [%s] result [%s]"),
-					*AssetName, 
-					VertexInstanceID.GetValue(),
-					*ReferenceVertexInstanceColors[VertexInstanceID].ToString(),
-					*ResultVertexInstanceColors[VertexInstanceID].ToString())));
-				bAllSame = false;
+				ExecutionInfo.AddEvent(FAutomationEvent(EAutomationEventType::Error, FString::Printf(TEXT("The %s conversion %s is not lossless, %s array is different. Array index [%d] expected %s [%d] result [%d]"),
+					*AssetName,
+					*ConversionName,
+					*VectorArrayName,
+					ElementID.GetValue(),
+					*VectorArrayName,
+					ReferenceArray[ElementID],
+					ResultArray[ElementID])));
+				bIsSame = false;
 				break;
 			}
 		}
 	}
+}
+
+bool FMeshDescriptionTest::CompareMeshDescription(const FString& AssetName, FAutomationTestExecutionInfo& ExecutionInfo, const UMeshDescription* ReferenceMeshDescription, const UMeshDescription* MeshDescription) const
+{
+	//////////////////////////////////////////////////////////////////////////
+	//Gather the reference data
+	const TVertexAttributeArray<FVector>& ReferenceVertexPositions = ReferenceMeshDescription->VertexAttributes().GetAttributes<FVector>(MeshAttribute::Vertex::Position);
+
+	const TVertexInstanceAttributeArray<FVector>& ReferenceVertexInstanceNormals = ReferenceMeshDescription->VertexInstanceAttributes().GetAttributes<FVector>(MeshAttribute::VertexInstance::Normal);
+	const TVertexInstanceAttributeArray<FVector>& ReferenceVertexInstanceTangents = ReferenceMeshDescription->VertexInstanceAttributes().GetAttributes<FVector>(MeshAttribute::VertexInstance::Tangent);
+	const TVertexInstanceAttributeArray<float>& ReferenceVertexInstanceBinormalSigns = ReferenceMeshDescription->VertexInstanceAttributes().GetAttributes<float>(MeshAttribute::VertexInstance::BinormalSign);
+	const TVertexInstanceAttributeArray<FVector4>& ReferenceVertexInstanceColors = ReferenceMeshDescription->VertexInstanceAttributes().GetAttributes<FVector4>(MeshAttribute::VertexInstance::Color);
+	const TVertexInstanceAttributeIndicesArray<FVector2D>& ReferenceVertexInstanceUVs = ReferenceMeshDescription->VertexInstanceAttributes().GetAttributesSet<FVector2D>(MeshAttribute::VertexInstance::TextureCoordinate);
+
+	const TEdgeAttributeArray<bool>& ReferenceEdgeHardnesses = ReferenceMeshDescription->EdgeAttributes().GetAttributes<bool>(MeshAttribute::Edge::IsHard);
+
+	const TPolygonGroupAttributeArray<int>& ReferencePolygonGroupMaterialIndex = ReferenceMeshDescription->PolygonGroupAttributes().GetAttributes<int>(MeshAttribute::PolygonGroup::MaterialIndex);
+
+
+	//////////////////////////////////////////////////////////////////////////
+	//Gather the result data
+	const TVertexAttributeArray<FVector>& ResultVertexPositions = MeshDescription->VertexAttributes().GetAttributes<FVector>(MeshAttribute::Vertex::Position);
+
+	const TVertexInstanceAttributeArray<FVector>& ResultVertexInstanceNormals = MeshDescription->VertexInstanceAttributes().GetAttributes<FVector>(MeshAttribute::VertexInstance::Normal);
+	const TVertexInstanceAttributeArray<FVector>& ResultVertexInstanceTangents = MeshDescription->VertexInstanceAttributes().GetAttributes<FVector>(MeshAttribute::VertexInstance::Tangent);
+	const TVertexInstanceAttributeArray<float>& ResultVertexInstanceBinormalSigns = MeshDescription->VertexInstanceAttributes().GetAttributes<float>(MeshAttribute::VertexInstance::BinormalSign);
+	const TVertexInstanceAttributeArray<FVector4>& ResultVertexInstanceColors = MeshDescription->VertexInstanceAttributes().GetAttributes<FVector4>(MeshAttribute::VertexInstance::Color);
+	const TVertexInstanceAttributeIndicesArray<FVector2D>& ResultVertexInstanceUVs = MeshDescription->VertexInstanceAttributes().GetAttributesSet<FVector2D>(MeshAttribute::VertexInstance::TextureCoordinate);
+
+	const TEdgeAttributeArray<bool>& ResultEdgeHardnesses = MeshDescription->EdgeAttributes().GetAttributes<bool>(MeshAttribute::Edge::IsHard);
+
+	const TPolygonGroupAttributeArray<int>& ResultPolygonGroupMaterialIndex = MeshDescription->PolygonGroupAttributes().GetAttributes<int>(MeshAttribute::PolygonGroup::MaterialIndex);
+
+	
+	//////////////////////////////////////////////////////////////////////////
+	// Do the comparison
+	bool bAllSame = true;
+
+	FString ConversionName = TEXT("MeshDescription to RawMesh to MeshDescription");
+
+	//Positions
+	MeshDescriptionStructureArrayCompare<TVertexAttributeArray<FVector>, FVertexArray, FVertexID>(ConversionName, AssetName, ExecutionInfo, bAllSame, ReferenceMeshDescription->Vertices(), TEXT("vertex positions"), ReferenceVertexPositions, ResultVertexPositions);
+
+	//Normals
+	MeshDescriptionStructureArrayCompare<TVertexInstanceAttributeArray<FVector>, FVertexInstanceArray, FVertexInstanceID>(ConversionName, AssetName, ExecutionInfo, bAllSame, ReferenceMeshDescription->VertexInstances(), TEXT("vertex instance normals"), ReferenceVertexInstanceNormals, ResultVertexInstanceNormals);
+
+	//Tangents
+	MeshDescriptionStructureArrayCompare<TVertexInstanceAttributeArray<FVector>, FVertexInstanceArray, FVertexInstanceID>(ConversionName, AssetName, ExecutionInfo, bAllSame, ReferenceMeshDescription->VertexInstances(), TEXT("vertex instance tangents"), ReferenceVertexInstanceTangents, ResultVertexInstanceTangents);
+
+	//BiNormal signs
+	MeshDescriptionNumberArrayCompare<TVertexInstanceAttributeArray<float>, FVertexInstanceArray, FVertexInstanceID>(ConversionName, AssetName, ExecutionInfo, bAllSame, ReferenceMeshDescription->VertexInstances(), TEXT("vertex instance binormals"), ReferenceVertexInstanceBinormalSigns, ResultVertexInstanceBinormalSigns);
+
+	//Colors
+	MeshDescriptionStructureArrayCompare<TVertexInstanceAttributeArray<FVector4>, FVertexInstanceArray, FVertexInstanceID>(ConversionName, AssetName, ExecutionInfo, bAllSame, ReferenceMeshDescription->VertexInstances(), TEXT("vertex instance colors"), ReferenceVertexInstanceColors, ResultVertexInstanceColors);
+
 	//Uvs
 	if (ReferenceVertexInstanceUVs.GetNumIndices() != ResultVertexInstanceUVs.GetNumIndices())
 	{
-		ExecutionInfo.AddEvent(FAutomationEvent(EAutomationEventType::Error, FString::Printf(TEXT("The %s conversion to RawMesh is not lossless, vertex instance UVs channel count is different. UVs channel count expected [%d] result [%d]"),
+		ExecutionInfo.AddEvent(FAutomationEvent(EAutomationEventType::Error, FString::Printf(TEXT("The %s conversion MeshDescription to RawMesh to MeshDescription is not lossless, vertex instance UVs channel count is different. UVs channel count expected [%d] result [%d]"),
 			*AssetName,
 			ReferenceVertexInstanceUVs.GetNumIndices(),
 			ResultVertexInstanceUVs.GetNumIndices())));
@@ -300,41 +376,18 @@ bool FMeshDescriptionTest::CompareMeshDescription(const FString& AssetName, FAut
 	{
 		for (int32 UVChannelIndex = 0; UVChannelIndex < ReferenceVertexInstanceUVs.GetNumIndices(); ++UVChannelIndex)
 		{
+			FString UVIndexName = FString::Printf(TEXT("vertex instance UVs(%d)"), UVChannelIndex);
 			const TMeshAttributeArray<FVector2D, FVertexInstanceID>& ReferenceUVs = ReferenceVertexInstanceUVs.GetArrayForIndex(UVChannelIndex);
 			const TMeshAttributeArray<FVector2D, FVertexInstanceID>& ResultUVs = ResultVertexInstanceUVs.GetArrayForIndex(UVChannelIndex);
-			if (ReferenceUVs.Num() != ResultUVs.Num())
-			{
-				ExecutionInfo.AddEvent(FAutomationEvent(EAutomationEventType::Error, FString::Printf(TEXT("The %s conversion to RawMesh is not lossless, vertex instance UV channel[%d] UVs count is different. Uvs count expected [%d] result [%d]"),
-					*AssetName,
-					UVChannelIndex,
-					ReferenceUVs.Num(),
-					ResultUVs.Num())));
-				bAllSame = false;
-			}
-			else
-			{
-				for (FVertexInstanceID VertexInstanceID : ReferenceMeshDescription->VertexInstances().GetElementIDs())
-				{
-					if (ReferenceUVs[VertexInstanceID] != ResultUVs[VertexInstanceID])
-					{
-						ExecutionInfo.AddEvent(FAutomationEvent(EAutomationEventType::Error, FString::Printf(TEXT("The %s conversion to RawMesh is not lossless, vertex instance UV Channel[%d] UVs array is different. VertexInstanceID [%d] expected UV [%s] result [%s]"),
-							*AssetName,
-							UVChannelIndex,
-							VertexInstanceID.GetValue(),
-							*ReferenceUVs[VertexInstanceID].ToString(),
-							*ResultUVs[VertexInstanceID].ToString())));
-						bAllSame = false;
-						break;
-					}
-				}
-			}
+			MeshDescriptionStructureArrayCompare<TMeshAttributeArray<FVector2D, FVertexInstanceID>, FVertexInstanceArray, FVertexInstanceID>(ConversionName, AssetName, ExecutionInfo, bAllSame, ReferenceMeshDescription->VertexInstances(), UVIndexName, ReferenceUVs, ResultUVs);
 		}
 	}
+
 	//Edges
-	//We check if hard edges are kept correctly
+	//We do not use a template since we need to check the connected polygon count to validate a false comparison
 	if (ReferenceEdgeHardnesses.Num() != ResultEdgeHardnesses.Num())
 	{
-		ExecutionInfo.AddEvent(FAutomationEvent(EAutomationEventType::Error, FString::Printf(TEXT("The %s conversion to RawMesh is not lossless, Edge count is different. Edges count expected [%d] result [%d]"),
+		ExecutionInfo.AddEvent(FAutomationEvent(EAutomationEventType::Error, FString::Printf(TEXT("The %s conversion MeshDescription to RawMesh to MeshDescription is not lossless, Edge count is different. Edges count expected [%d] result [%d]"),
 			*AssetName,
 			ReferenceEdgeHardnesses.Num(),
 			ResultEdgeHardnesses.Num())));
@@ -358,36 +411,14 @@ bool FMeshDescriptionTest::CompareMeshDescription(const FString& AssetName, FAut
 						(ResultEdgeHardnesses[EdgeID] ? TEXT("true") : TEXT("false")))));
 					bAllSame = false;
 				}
-				//break;
-			}
-		}
-	}
-	//Polygon group ID
-	//We currently rely only on the PolygonGroupID. The duplicate material slot information (the info is store in the staticmesh material array) is not necessary and cannot be put in FRawMesh structure.
-	if (ReferenceMeshDescription->PolygonGroups().Num() != MeshDescription->PolygonGroups().Num())
-	{
-		ExecutionInfo.AddEvent(FAutomationEvent(EAutomationEventType::Error, FString::Printf(TEXT("The %s conversion to RawMesh is not lossless, Polygon group count is different. Polygon group count expected [%d] result [%d]"),
-			*AssetName,
-			ReferenceMeshDescription->PolygonGroups().Num(),
-			MeshDescription->PolygonGroups().Num())));
-		bAllSame = false;
-	}
-	else
-	{
-		for (const FPolygonGroupID& PolygonGroupID : ReferenceMeshDescription->PolygonGroups().GetElementIDs())
-		{
-			if (ReferencePolygonGroupMaterialIndex[PolygonGroupID] != ResultPolygonGroupMaterialIndex[PolygonGroupID])
-			{
-				ExecutionInfo.AddEvent(FAutomationEvent(EAutomationEventType::Error, FString::Printf(TEXT("The %s conversion to RawMesh is not lossless, polygon group material index array is different. PolygonGroupID [%d] expected Material index [%d] result [%d]"),
-					*AssetName,
-					PolygonGroupID.GetValue(),
-					ReferencePolygonGroupMaterialIndex[PolygonGroupID],
-					ResultPolygonGroupMaterialIndex[PolygonGroupID])));
-				bAllSame = false;
 				break;
 			}
 		}
 	}
+
+	//Polygon group ID
+	MeshDescriptionNumberArrayCompare<TPolygonGroupAttributeArray<int>, FPolygonGroupArray, FPolygonGroupID>(ConversionName, AssetName, ExecutionInfo, bAllSame, ReferenceMeshDescription->PolygonGroups(), TEXT("vertex instance binormals"), ReferencePolygonGroupMaterialIndex, ResultPolygonGroupMaterialIndex);
+
 	return bAllSame;
 }
 
@@ -421,21 +452,56 @@ bool FMeshDescriptionTest::ConversionTest(FAutomationTestExecutionInfo& Executio
 
 		if (AssetMesh != nullptr)
 		{
-			const UMeshDescription* ReferenceAssetMesh = AssetMesh->GetOriginalMeshDescription(0);
-			check(ReferenceAssetMesh != nullptr);
-			//Create a temporary Mesh Description
-			UMeshDescription* ResultAssetMesh = DuplicateObject<UMeshDescription>(ReferenceAssetMesh, GetTransientPackage(), NAME_None);
-			//Convert MeshDescription to FRawMesh
-			FRawMesh RawMesh;
-			FMeshDescriptionOperations::ConverToRawMesh(ResultAssetMesh, RawMesh);
-			//Convert back the FRawmesh
-			FMeshDescriptionOperations::ConverFromRawMesh(RawMesh, ResultAssetMesh);
-			if (!CompareMeshDescription(AssetName, ExecutionInfo, ReferenceAssetMesh, ResultAssetMesh))
+			//MeshDescription to RawMesh to MeshDescription
+			for(int32 LodIndex = 0; LodIndex < AssetMesh->SourceModels.Num(); ++LodIndex)
 			{
-				bAllSame = false;
+				const UMeshDescription* ReferenceAssetMesh = AssetMesh->GetOriginalMeshDescription(LodIndex);
+				if (ReferenceAssetMesh == nullptr)
+				{
+					check(LodIndex != 0);
+					continue;
+				}
+				//Create a temporary Mesh Description
+				UMeshDescription* ResultAssetMesh = DuplicateObject<UMeshDescription>(ReferenceAssetMesh, GetTransientPackage(), NAME_None);
+				//Convert MeshDescription to FRawMesh
+				FRawMesh RawMesh;
+				FMeshDescriptionOperations::ConverToRawMesh(ResultAssetMesh, RawMesh);
+				//Convert back the FRawmesh
+				FMeshDescriptionOperations::ConverFromRawMesh(RawMesh, ResultAssetMesh);
+				if (!CompareMeshDescription(AssetName, ExecutionInfo, ReferenceAssetMesh, ResultAssetMesh))
+				{
+					bAllSame = false;
+				}
+				//Destroy temporary object
+				ResultAssetMesh->ClearFlags(RF_Standalone);
+				ResultAssetMesh->MarkPendingKill();
+				
+			}
+			//RawMesh to MeshDescription to RawMesh
+			for (int32 LodIndex = 0; LodIndex < AssetMesh->SourceModels.Num(); ++LodIndex)
+			{
+				if (AssetMesh->SourceModels[LodIndex].RawMeshBulkData->IsEmpty())
+				{
+					check(LodIndex != 0);
+					continue;
+				}
+				FRawMesh ReferenceRawMesh;
+				AssetMesh->SourceModels[LodIndex].LoadRawMesh(ReferenceRawMesh);
+				FRawMesh ResultRawMesh;
+				AssetMesh->SourceModels[LodIndex].LoadRawMesh(ResultRawMesh);
+				//Create a temporary Mesh Description
+				UMeshDescription* MeshDescription = NewObject<UMeshDescription>(GetTransientPackage(), NAME_None, RF_Standalone);
+				FMeshDescriptionOperations::ConverFromRawMesh(ResultRawMesh, MeshDescription);
+				//Convert back the FRawmesh
+				FMeshDescriptionOperations::ConverToRawMesh(MeshDescription, ResultRawMesh);
+				if (!CompareRawMesh(AssetName, ExecutionInfo, ReferenceRawMesh, ResultRawMesh))
+				{
+					bAllSame = false;
+				}
+				//Destroy temporary object
+				MeshDescription->ClearFlags(RF_Standalone);
+				MeshDescription->MarkPendingKill();
 			}
-			//Destroy temporary object
-			ResultAssetMesh->MarkPendingKill();
 		}
 	}
 	//Collect garbage pass to remove everything from memory
diff --git a/Engine/Source/Editor/UnrealEd/Classes/Settings/EditorExperimentalSettings.h b/Engine/Source/Editor/UnrealEd/Classes/Settings/EditorExperimentalSettings.h
--- a/Engine/Source/Editor/UnrealEd/Classes/Settings/EditorExperimentalSettings.h
+++ b/Engine/Source/Editor/UnrealEd/Classes/Settings/EditorExperimentalSettings.h
@@ -125,7 +125,7 @@ public:
 	UPROPERTY(EditAnywhere, config, Category = Tools, meta = (ConfigRestartRequired = true))
 	bool bAssetMaterialBaking;
 
-	/** Allows usage of the mesh description instead of the raw mesh */
+	/** Allows usage of the mesh description instead of the raw mesh when we build static mesh */
 	UPROPERTY(EditAnywhere, config, Category = Tools, meta = (DisplayName = "Use Mesh Description"))
 	bool bUseMeshDescription;
 	
diff --git a/Engine/Source/Runtime/Engine/Classes/Engine/StaticMesh.h b/Engine/Source/Runtime/Engine/Classes/Engine/StaticMesh.h
--- a/Engine/Source/Runtime/Engine/Classes/Engine/StaticMesh.h
+++ b/Engine/Source/Runtime/Engine/Classes/Engine/StaticMesh.h
@@ -167,7 +167,7 @@ struct FStaticMeshSourceModel
 
 #if WITH_EDITORONLY_DATA
 	/* Original Imported mesh description data. Optional for all but the first LOD. Autogenerate LOD do not have original mesh description*/
-	UPROPERTY()
+	UPROPERTY(Transient)
 	class UMeshDescription* OriginalMeshDescription;
 #endif
 
@@ -958,19 +958,5 @@ private:
 	 * Fixes up the material when it was converted to the new staticmesh build process
 	 */
 	bool CleanUpRedondantMaterialPostLoad;
-
-	/*
-	 * Mesh data key DDC version enum
-	 */
-	enum EMeshDataKeyVersion
-	{
-		VER_MDK_INITIAL = 0,
-
-		// -----<new versions can be added before this line>-------------------------------------------------
-		// - this needs to be the last line (see note below)
-		VER_MDK_AUTOMATIC_VERSION_PLUS_ONE,
-		VER_MDK_AUTOMATIC_VERSION = VER_MDK_AUTOMATIC_VERSION_PLUS_ONE - 1
-	};
-
 #endif // #if WITH_EDITOR
 };
diff --git a/Engine/Source/Runtime/Engine/Private/StaticMesh.cpp b/Engine/Source/Runtime/Engine/Private/StaticMesh.cpp
--- a/Engine/Source/Runtime/Engine/Private/StaticMesh.cpp
+++ b/Engine/Source/Runtime/Engine/Private/StaticMesh.cpp
@@ -2366,50 +2366,60 @@ void UStaticMesh::FixupMaterialSlotName()
 	}
 }
 
+// If static mesh derived data needs to be rebuilt (new format, serialization
+// differences, etc.) replace the version GUID below with a new one.
+// In case of merge conflicts with DDC versions, you MUST generate a new GUID
+// and set this new GUID as the version.                                       
+#define MESHDATAKEY_STATICMESH_DERIVEDDATA_VER TEXT("81EFE7970F3140CABD7B24620A2BE2BE")
+
+static const FString& GetMeshDataKeyStaticMeshDerivedDataVersion()
+{
+	static FString CachedVersionString;
+	if (CachedVersionString.IsEmpty())
+	{
+		// Static mesh versioning is controlled by the version reported by the mesh utilities module.
+		CachedVersionString = FString::Printf(TEXT("%s%s"),
+			*UMeshDescription::GetMeshDescriptionVersion(),
+			MESHDATAKEY_STATICMESH_DERIVEDDATA_VER);
+	}
+	return CachedVersionString;
+}
+
 bool UStaticMesh::GetMeshDataKey(FString& OutKey)
 {
+	OutKey.Empty();
 	if (SourceModels.Num() < 1)
 	{
 		return false;
 	}
-	OutKey += TEXT("MESHDATAKEY_STATICMESH___Path_");
-	OutKey += this->GetFullGroupName(false);
-	OutKey += TEXT("___Version_");
-	OutKey += FString::FromInt(VER_MDK_AUTOMATIC_VERSION+1);
-	OutKey += TEXT("___LOD_COUNT_");
-	OutKey += FString::FromInt(SourceModels.Num());
 	FSHA1 Sha;
 	for (int32 LodIndex = 0; LodIndex < SourceModels.Num(); ++LodIndex)
 	{
 		FString LodIndexString = FString::Printf(TEXT("%d_"), LodIndex);
 		FStaticMeshSourceModel& SourceModel = SourceModels[LodIndex];
-		bool bEmptyRawData = true;
-		if (SourceModel.OriginalMeshDescription != nullptr)
-		{
-			LodIndexString += SourceModel.OriginalMeshDescription->GetIdString();
-			bEmptyRawData = false;
-		}
-		
 		if(!SourceModel.RawMeshBulkData->IsEmpty())
 		{
 			LodIndexString += SourceModel.RawMeshBulkData->GetIdString();
-			bEmptyRawData = false;
 		}
-		
-		if(bEmptyRawData)
+		else
 		{
 			LodIndexString += TEXT("REDUCELOD");
 		}
 		TArray<TCHAR> LodIndexArray = LodIndexString.GetCharArray();
 		Sha.Update((uint8*)LodIndexArray.GetData(), LodIndexArray.Num() * LodIndexArray.GetTypeSize());
 	}
-	OutKey += TEXT("___DATAHASH_");
 	Sha.Final();
 	// Retrieve the hash and use it to construct a pseudo-GUID.
 	uint32 Hash[5];
 	Sha.GetHash((uint8*)Hash);
 	FGuid Guid = FGuid(Hash[0] ^ Hash[4], Hash[1], Hash[2], Hash[3]);
-	OutKey += Guid.ToString(EGuidFormats::Digits);
+	FString MeshLodData = Guid.ToString(EGuidFormats::Digits);
+
+	OutKey = FDerivedDataCacheInterface::BuildCacheKey(
+		TEXT("MESHDATAKEY_STATICMESH"),
+		*GetMeshDataKeyStaticMeshDerivedDataVersion(),
+		*MeshLodData
+	);
 	return true;
 }
 
@@ -2417,8 +2427,6 @@ struct FMeshDataSerialize
 {
 private:
 	int32 LodNumber;
-	TArray<bool> RawMeshTable;
-	TArray<FRawMesh> RawMeshContent;
 	TArray<bool> MeshDescriptionTable;
 	TArray<UMeshDescription*> MeshDescriptionContent;
 
@@ -2439,8 +2447,6 @@ public:
 			}
 		}
 		LodNumber = 0;
-		RawMeshTable.Empty();
-		RawMeshContent.Empty();
 		MeshDescriptionTable.Empty();
 		MeshDescriptionContent.Empty();
 	}
@@ -2449,19 +2455,12 @@ public:
 	void ReadDataFromStaticMesh(UStaticMesh* StaticMesh)
 	{
 		LodNumber = StaticMesh->SourceModels.Num();
-		RawMeshTable.AddZeroed(LodNumber);
-		RawMeshContent.AddDefaulted(LodNumber);
 		MeshDescriptionTable.AddZeroed(LodNumber);
 		MeshDescriptionContent.AddZeroed(LodNumber);
 		for (int32 LodIndex = 0; LodIndex < LodNumber; ++LodIndex)
 		{
 			StaticMesh->SourceModels[LodIndex];
 			FStaticMeshSourceModel& SourceModel = StaticMesh->SourceModels[LodIndex];
-			RawMeshTable[LodIndex] = !SourceModel.RawMeshBulkData->IsEmpty();
-			if (RawMeshTable[LodIndex])
-			{
-				SourceModel.RawMeshBulkData->LoadRawMesh(RawMeshContent[LodIndex]);
-			}
 			MeshDescriptionTable[LodIndex] = SourceModel.OriginalMeshDescription != nullptr;
 			if (MeshDescriptionTable[LodIndex])
 			{
@@ -2472,20 +2471,15 @@ public:
 	}
 
 	//Apply the structure data to a static mesh
-	bool ApplyDataToStaticMesh(UStaticMesh* StaticMesh, int32 LodIndex, bool bRawMesh)
+	//Return true if we have valid data to set in the static mesh
+	bool WriteDataToStaticMesh(UStaticMesh* StaticMesh, int32 LodIndex)
 	{
 		if (LodIndex >= LodNumber || LodIndex < 0 || !StaticMesh->SourceModels.IsValidIndex(LodIndex))
 		{
 			return false;
 		}
 
-		if (bRawMesh && RawMeshTable.IsValidIndex(LodIndex) && RawMeshTable[LodIndex] && RawMeshContent[LodIndex].IsValid())
-		{
-			StaticMesh->SourceModels[LodIndex].RawMeshBulkData->SaveRawMesh(RawMeshContent[LodIndex]);
-			return true;
-		}
-
-		if (!bRawMesh && MeshDescriptionTable.IsValidIndex(LodIndex) && MeshDescriptionTable[LodIndex] && MeshDescriptionContent[LodIndex] != nullptr)
+		if (MeshDescriptionTable.IsValidIndex(LodIndex) && MeshDescriptionTable[LodIndex] && MeshDescriptionContent[LodIndex] != nullptr)
 		{
 			StaticMesh->SourceModels[LodIndex].OriginalMeshDescription = Cast<UMeshDescription>(StaticDuplicateObject(MeshDescriptionContent[LodIndex], StaticMesh, NAME_None, RF_AllFlags & ~RF_Standalone));
 			return true;
@@ -2498,22 +2492,10 @@ public:
 	{
 		FMemoryReader Ar(DerivedData, /*bIsPersistent=*/ true);
 		Ar.Serialize((void*)(&LodNumber), sizeof(int32));
-		RawMeshTable.AddZeroed(LodNumber);
-		RawMeshContent.AddDefaulted(LodNumber);
 		MeshDescriptionTable.AddZeroed(LodNumber);
 		MeshDescriptionContent.AddZeroed(LodNumber);
 		for (int32 LodIndex = 0; LodIndex < LodNumber; ++LodIndex)
 		{
-			int32 HasRawMesh;
-			Ar.Serialize((void*)(&HasRawMesh), sizeof(int32));
-			RawMeshTable[LodIndex] = HasRawMesh == 0 ? false : true;
-			
-			if (RawMeshTable[LodIndex])
-			{
-				//Read the FRawMesh
-				Ar << RawMeshContent[LodIndex];
-			}
-
 			int32 HasMeshDescription;
 			Ar.Serialize((void*)(&HasMeshDescription), sizeof(int32));
 			MeshDescriptionTable[LodIndex] = HasMeshDescription == 0 ? false : true;
@@ -2532,14 +2514,6 @@ public:
 		Ar.Serialize((void*)(&LodNumber), sizeof(int32));
 		for (int32 LodIndex = 0; LodIndex < LodNumber; ++LodIndex)
 		{
-			int32 HasRawMesh = RawMeshTable[LodIndex] ? 1 : 0;
-			Ar.Serialize((void*)(&HasRawMesh), sizeof(int32));
-			if (HasRawMesh != 0)
-			{
-				//Write the FRawMeshBulkData
-				Ar << RawMeshContent[LodIndex];
-			}
-
 			int32 HasMeshDescription = MeshDescriptionTable[LodIndex] ? 1 : 0;
 			Ar.Serialize((void*)(&HasMeshDescription), sizeof(int32));
 			if (HasMeshDescription)
@@ -2553,10 +2527,9 @@ public:
 
 void UStaticMesh::CacheMeshData()
 {
-	bool bMeshDescriptionMode = GetDefault<UEditorExperimentalSettings>()->bUseMeshDescription;
-	//TODO DDC
 	//Look if a ddc exist for this mesh data
 	bool IsValidCache = false;
+	//Since we can have more then 1 LOD we need a structure to store the RawData per LOD to read/write the mesh description ddc
 	FMeshDataSerialize MeshDataSerialize;
 	FString MeshDataKey;
 	if (GetMeshDataKey(MeshDataKey))
@@ -2570,86 +2543,29 @@ void UStaticMesh::CacheMeshData()
 		}
 	}
 
-	//We have two lambda, one to create a raw mesh if the mesh description exist and another lambda to create a mesh description if the raw mesh exist
-	//We use lambda because we want to look for them in a order decide by the mesh description mode
-
-	//Create a raw mesh from a valid mesh description
-	auto MeshDescriptionLambda = [this, &IsValidCache, &MeshDataSerialize](int32 LodIndex)
+	for (int32 LodIndex = 0; LodIndex < SourceModels.Num(); ++LodIndex)
 	{
+		//For phase 1 the mesh description is transient to avoid loosing data
 		FStaticMeshSourceModel& SourceModel = SourceModels[LodIndex];
-		if (SourceModel.OriginalMeshDescription != nullptr)
+		if (!SourceModel.RawMeshBulkData->IsEmpty() && SourceModel.OriginalMeshDescription == nullptr)
 		{
-			if (SourceModel.RawMeshBulkData->IsEmpty())
+			//Use the ddc if it exist, otherwise convert the data
+			if (!MeshDataSerialize.WriteDataToStaticMesh(this, LodIndex))
 			{
-				//Convert the mesh description to raw mesh and mark the ddc for save
+				//ddc does not contain a mesh description, create one
+				IsValidCache = false;
 				FRawMesh TempRawMesh;
-				//Use the ddc if it exist
-				if (!MeshDataSerialize.ApplyDataToStaticMesh(this, LodIndex, true))
-				{
-					IsValidCache = false;
-					//Convert RawMesh to meshdescription
-					FMeshDescriptionOperations::ConverToRawMesh(SourceModel.OriginalMeshDescription, TempRawMesh);
-				}
-				SourceModel.RawMeshBulkData->SaveRawMesh(TempRawMesh);
+				//Convert the rawmesh to mesh description mark the ddc for save
+				SourceModel.RawMeshBulkData->LoadRawMesh(TempRawMesh);
+				//The original mesh description is null we must create one
+				SourceModel.OriginalMeshDescription = NewObject<UMeshDescription>(this, NAME_None);
+				//Convert the raw mesh to mesh description and save the mesh description in the ddc
+				FMeshDescriptionOperations::ConverFromRawMesh(TempRawMesh, SourceModel.OriginalMeshDescription);
 			}
-			return true;
 		}
-		return false;
-	};
-
-	//Create a mesh description from a valid raw mesh
-	auto RawMeshLambda = [this, &IsValidCache, &MeshDataSerialize](int32 LodIndex)
-	{
-		FStaticMeshSourceModel& SourceModel = SourceModels[LodIndex];
-		if (!SourceModel.RawMeshBulkData->IsEmpty())
-		{
-			if (SourceModel.OriginalMeshDescription == nullptr)
-			{
-				//Use the ddc if it exist
-				if (!MeshDataSerialize.ApplyDataToStaticMesh(this, LodIndex, false))
-				{
-					IsValidCache = false;
-					FRawMesh TempRawMesh;
-					//Convert the rawmesh to mesh description mark the ddc for save
-					SourceModel.RawMeshBulkData->LoadRawMesh(TempRawMesh);
-					//The original mesh description is null we must create one
-					SourceModel.OriginalMeshDescription = NewObject<UMeshDescription>(this, NAME_None);
-					//Convert the raw mesh to mesh description and save the mesh description in the ddc
-					FMeshDescriptionOperations::ConverFromRawMesh(TempRawMesh, SourceModel.OriginalMeshDescription);
-				}
-			}
-			return true;
-		}
-		return false;
-	};
-
-	for (int32 LodIndex = 0; LodIndex < SourceModels.Num(); ++LodIndex)
-	{
-		if (bMeshDescriptionMode)
-		{
-			if (!MeshDescriptionLambda(LodIndex))
-			{
-				RawMeshLambda(LodIndex);
-			}
-		}
-		else
-		{
-			if (!RawMeshLambda(LodIndex))
-			{
-				MeshDescriptionLambda(LodIndex);
-			}
-		}
-
-		//We cannot check because there is some case where the MeshMerge as generate some mesh LOD with no rawmesh and no reduction
-		//if (LodIndex == 0 || (SourceModels[LodIndex].ReductionSettings.PercentTriangles == 1.0f && SourceModels[LodIndex].ReductionSettings.MaxDeviation == 0.0f))
-		//{
-			//Check if we have all the data, we need both data until we deprecate FRawMesh
-			//check(SourceModels[LodIndex].OriginalMeshDescription != nullptr);
-			//check(!SourceModels[LodIndex].RawMeshBulkData->IsEmpty());
-		//}
 	}
 
-	//Write a cache that will be useful next time we load the asset (if not saved)
+	//Write the mesh description ddc cache if necessary
 	if (!IsValidCache && !MeshDataKey.IsEmpty())
 	{
 		//We need to write the cache
diff --git a/Engine/Source/Runtime/MeshBuilder/Private/MeshDescriptionHelper.cpp b/Engine/Source/Runtime/MeshBuilder/Private/MeshDescriptionHelper.cpp
--- a/Engine/Source/Runtime/MeshBuilder/Private/MeshDescriptionHelper.cpp
+++ b/Engine/Source/Runtime/MeshBuilder/Private/MeshDescriptionHelper.cpp
@@ -929,7 +929,7 @@ void FMeshDescriptionHelper::ConvertSmoothGroupToHardEdges(const struct FRawMesh
 void FMeshDescriptionHelper::ConverToRawMesh(const UMeshDescription* SourceMeshDescription, struct FRawMesh &DestinationRawMesh)
 {
 	static BuildStatisticManager::FStatisticData StatisticData;
-	BuildStatisticManager::FBuildStatisticScope(TEXT("ConvertToRawMesh"), StatisticData);
+	BuildStatisticManager::FBuildStatisticScope StatScope(TEXT("ConvertToRawMesh"), StatisticData);
 	DestinationRawMesh.Empty();
 
 	//Gather all array data
@@ -1001,8 +1001,13 @@ void FMeshDescriptionHelper::ConverToRawMesh(const UMeshDescription* SourceMeshD
 void FMeshDescriptionHelper::ConverFromRawMesh(const struct FRawMesh &SourceRawMesh, UMeshDescription* DestinationMeshDescription)
 {
 	static BuildStatisticManager::FStatisticData StatisticData;
-	BuildStatisticManager::FBuildStatisticScope(TEXT("ConverFromRawMesh"), StatisticData);
+	BuildStatisticManager::FBuildStatisticScope StatScope(TEXT("ConverFromRawMesh"), StatisticData);
 	DestinationMeshDescription->Empty();
+
+	DestinationMeshDescription->ReserveNewVertices(SourceRawMesh.VertexPositions.Num());
+	DestinationMeshDescription->ReserveNewVertexInstances(SourceRawMesh.WedgeIndices.Num());
+	DestinationMeshDescription->ReserveNewPolygons(SourceRawMesh.WedgeIndices.Num() / 3);
+	DestinationMeshDescription->ReserveNewEdges(SourceRawMesh.WedgeIndices.Num() * 2 / 3); //Approximately 2 edges per polygons
 	//Gather all array data
 	TVertexAttributeArray<FVector>& VertexPositions = DestinationMeshDescription->VertexAttributes().GetAttributes<FVector>(MeshAttribute::Vertex::Position);
 
@@ -1109,8 +1114,6 @@ void FMeshDescriptionHelper::ConverFromRawMesh(const struct FRawMesh &SourceRawM
 			}
 			ContourPoint.EdgeID = MatchEdgeId;
 			ContourPoint.VertexInstanceID = FVertexInstanceID(VerticeIndexBase + CornerIndices[0]);
-
-			//TODO Edges smoothing
 		}
 
 		const FPolygonID NewPolygonID = DestinationMeshDescription->CreatePolygon(PolygonGroupID, Contours);
diff --git a/Engine/Source/Runtime/MeshDescription/Private/MeshDescription.cpp b/Engine/Source/Runtime/MeshDescription/Private/MeshDescription.cpp
--- a/Engine/Source/Runtime/MeshDescription/Private/MeshDescription.cpp
+++ b/Engine/Source/Runtime/MeshDescription/Private/MeshDescription.cpp
@@ -76,7 +76,6 @@ void UMeshDescription::UnRegisterBaseAttributes()
 void UMeshDescription::Serialize( FArchive& Ar )
 {
 	Super::Serialize( Ar );
-
 	Ar << VertexArray;
 	Ar << VertexInstanceArray;
 	Ar << EdgeArray;
@@ -106,15 +105,150 @@ void UMeshDescription::Empty()
 #if WITH_EDITORONLY_DATA
 FString UMeshDescription::GetIdString()
 {
-	TArray<uint8> TempBytes;
-	FObjectWriter Ar(this, TempBytes);
-	Serialize(Ar);
+	//The serialisation of the sparse array can be different with the same data, so to get a unique identifier we need to take the data with no IDs(FVertexID...)
 	FSHA1 Sha;
-	if (TempBytes.Num() > 0)
+	TArray<uint8> TempBytes;
+
+	TVertexAttributeArray<FVector>& VertexPositions = VertexAttributes().GetAttributes<FVector>(MeshAttribute::Vertex::Position);
+	TVertexAttributeArray<float>& VertexCornerSharpness = VertexAttributes().GetAttributes<float>(MeshAttribute::Vertex::CornerSharpness);
+
+	TVertexInstanceAttributeIndicesArray<FVector2D>& TextureCoordinates = VertexInstanceAttributes().GetAttributesSet<FVector2D>(MeshAttribute::VertexInstance::TextureCoordinate);
+	TVertexInstanceAttributeArray<FVector>& VertexInstanceNormals = VertexInstanceAttributes().GetAttributes<FVector>(MeshAttribute::VertexInstance::Normal);
+	TVertexInstanceAttributeArray<FVector>& VertexInstanceTangents = VertexInstanceAttributes().GetAttributes<FVector>(MeshAttribute::VertexInstance::Tangent);
+	TVertexInstanceAttributeArray<float>& VertexInstanceBiNormals = VertexInstanceAttributes().GetAttributes<float>(MeshAttribute::VertexInstance::BinormalSign);
+	TVertexInstanceAttributeArray<FVector4>& VertexInstanceColors = VertexInstanceAttributes().GetAttributes<FVector4>(MeshAttribute::VertexInstance::Color);
+
+	TEdgeAttributeArray<bool>& EdgeHards = EdgeAttributes().GetAttributes<bool>(MeshAttribute::Edge::IsHard);
+	TEdgeAttributeArray<float>& EdgeCreaseSharpness = EdgeAttributes().GetAttributes<float>(MeshAttribute::Edge::CreaseSharpness);
+
+	TPolygonAttributeArray<FVector>& PolygonNormals = PolygonAttributes().GetAttributes<FVector>(MeshAttribute::Polygon::Normal);
+	TPolygonAttributeArray<FVector>& PolygonTangents = PolygonAttributes().GetAttributes<FVector>(MeshAttribute::Polygon::Tangent);
+	TPolygonAttributeArray<FVector>& PolygonBinormals = PolygonAttributes().GetAttributes<FVector>(MeshAttribute::Polygon::Binormal);
+	TPolygonAttributeArray<FVector>& PolygonCenters = PolygonAttributes().GetAttributes<FVector>(MeshAttribute::Polygon::Center);
+
+	TPolygonGroupAttributeArray<int>& PolygonGroupMaterialIndexes = PolygonGroupAttributes().GetAttributes<int>(MeshAttribute::PolygonGroup::MaterialIndex);
+	TPolygonGroupAttributeArray<FName>& PolygonGroupMaterialSlotNames = PolygonGroupAttributes().GetAttributes<FName>(MeshAttribute::PolygonGroup::MaterialSlotName);
+	TPolygonGroupAttributeArray<FName>& PolygonGroupImportedMaterialSlotNames = PolygonGroupAttributes().GetAttributes<FName>(MeshAttribute::PolygonGroup::ImportedMaterialSlotName);
+	TPolygonGroupAttributeArray<bool>& PolygonGroupEnableCollisions = PolygonGroupAttributes().GetAttributes<bool>(MeshAttribute::PolygonGroup::EnableCollision);
+	TPolygonGroupAttributeArray<bool>& PolygonGroupCastShadows = PolygonGroupAttributes().GetAttributes<bool>(MeshAttribute::PolygonGroup::CastShadow);
+	
+	auto AddItemToSha = [&](auto ItemSerialization)
 	{
-		uint8* Buffer = TempBytes.GetData();
-		Sha.Update(Buffer, TempBytes.Num());
-	}
+		FMemoryWriter Ar(TempBytes);
+		ItemSerialization(Ar);
+		if (TempBytes.Num() > 0)
+		{
+			uint8* Buffer = TempBytes.GetData();
+			Sha.Update(Buffer, TempBytes.Num());
+		}
+		TempBytes.Reset();
+	};
+
+	//Vertex Identifier
+	auto VertexSerializer = [&](FMemoryWriter &Ar)
+	{
+		for (const FVertexID& ArrayID : Vertices().GetElementIDs())
+		{
+			const FMeshVertex& MeshItem = GetVertex(ArrayID);
+			uint32 ItemCount = MeshItem.ConnectedEdgeIDs.Num();
+			Ar.SerializeInt(ItemCount, 0);
+			ItemCount = MeshItem.VertexInstanceIDs.Num();
+			Ar.SerializeInt(ItemCount, 0);
+			Ar << VertexPositions[ArrayID];
+			Ar << VertexCornerSharpness[ArrayID];
+		}
+	};
+	AddItemToSha(VertexSerializer);
+	
+	//Vertex Instance Identifier
+	auto VertexInstanceSerializer = [&](FMemoryWriter &Ar)
+	{
+		for (const FVertexInstanceID& ArrayID : VertexInstances().GetElementIDs())
+		{
+			const FMeshVertexInstance& MeshItem = GetVertexInstance(ArrayID);
+			uint32 ItemCount = MeshItem.ConnectedPolygons.Num();
+			Ar.SerializeInt(ItemCount, 0);
+			Ar << VertexPositions[MeshItem.VertexID];
+			for (int32 UVIndex = 0; UVIndex < TextureCoordinates.GetNumIndices(); ++UVIndex)
+			{
+				Ar << TextureCoordinates.GetArrayForIndex(UVIndex)[ArrayID];
+			}
+			Ar << VertexInstanceNormals[ArrayID];
+			Ar << VertexInstanceTangents[ArrayID];
+			Ar << VertexInstanceBiNormals[ArrayID];
+			Ar << VertexInstanceColors[ArrayID];
+		}
+	};
+	AddItemToSha(VertexInstanceSerializer);
+
+	//Edge Identifier
+	auto EdgeSerializer = [&](FMemoryWriter &Ar)
+	{
+		for (const FEdgeID& ArrayID : Edges().GetElementIDs())
+		{
+			const FMeshEdge& MeshItem = GetEdge(ArrayID);
+			uint32 ItemCount = MeshItem.ConnectedPolygons.Num();
+			Ar.SerializeInt(ItemCount, 0);
+			Ar << VertexPositions[MeshItem.VertexIDs[0]];
+			Ar << VertexPositions[MeshItem.VertexIDs[1]];
+			Ar << EdgeHards[ArrayID];
+			Ar << EdgeCreaseSharpness[ArrayID];
+		}
+	};
+	AddItemToSha(EdgeSerializer);
+
+	//Polygon Identifier
+	auto PolygonSerializer = [&](FMemoryWriter &Ar)
+	{
+		for (const FPolygonID& ArrayID : Polygons().GetElementIDs())
+		{
+			const FMeshPolygon& MeshItem = GetPolygon(ArrayID);
+			uint32 ItemCount = MeshItem.HoleContours.Num();
+			Ar.SerializeInt(ItemCount, 0);
+			for(const FMeshPolygonContour& MeshPolygonContour : MeshItem.HoleContours)
+			{
+				ItemCount = MeshPolygonContour.VertexInstanceIDs.Num();
+				Ar.SerializeInt(ItemCount, 0);
+				for (const FVertexInstanceID& VertexInstanceID : MeshPolygonContour.VertexInstanceIDs)
+				{
+					Ar << VertexPositions[GetVertexInstance(VertexInstanceID).VertexID];
+				}
+			}
+			ItemCount = MeshItem.PerimeterContour.VertexInstanceIDs.Num();
+			Ar.SerializeInt(ItemCount, 0);
+			for (const FVertexInstanceID& VertexInstanceID : MeshItem.PerimeterContour.VertexInstanceIDs)
+			{
+				Ar << VertexPositions[GetVertexInstance(VertexInstanceID).VertexID];
+			}
+			ItemCount = PolygonGroupMaterialIndexes[MeshItem.PolygonGroupID];
+			Ar.SerializeInt(ItemCount, 0);
+
+			Ar << GetPolygon(ArrayID);
+			Ar << PolygonNormals[ArrayID];
+			Ar << PolygonTangents[ArrayID];
+			Ar << PolygonBinormals[ArrayID];
+			Ar << PolygonCenters[ArrayID];
+		}
+	};
+	AddItemToSha(PolygonSerializer);
+
+	//PolygonGroup Identifier
+	auto PolygonGroupSerializer = [&](FMemoryWriter &Ar)
+	{
+		for (const FPolygonGroupID& ArrayID : PolygonGroups().GetElementIDs())
+		{
+			const FMeshPolygonGroup& MeshItem = GetPolygonGroup(ArrayID);
+			uint32 ItemCount = MeshItem.Polygons.Num();
+			Ar.SerializeInt(ItemCount, 0);
+			Ar << PolygonGroupMaterialIndexes[ArrayID];
+			Ar << PolygonGroupMaterialSlotNames[ArrayID];
+			Ar << PolygonGroupImportedMaterialSlotNames[ArrayID];
+			Ar << PolygonGroupEnableCollisions[ArrayID];
+			Ar << PolygonGroupCastShadows[ArrayID];
+		}
+	};
+	AddItemToSha(PolygonGroupSerializer);
+	
 	Sha.Final();
 	TempBytes.Empty();
 	// Retrieve the hash and use it to construct a pseudo-GUID.
@@ -123,8 +257,7 @@ FString UMeshDescription::GetIdString()
 	FGuid Guid = FGuid(Hash[0] ^ Hash[4], Hash[1], Hash[2], Hash[3]);
 	return Guid.ToString(EGuidFormats::Digits);
 }
-#endif
-
+#endif //WITH_EDITORONLY_DATA
 
 void UMeshDescription::Compact( FElementIDRemappings& OutRemappings )
 {
diff --git a/Engine/Source/Runtime/MeshDescription/Public/MeshDescription.h b/Engine/Source/Runtime/MeshDescription/Public/MeshDescription.h
--- a/Engine/Source/Runtime/MeshDescription/Public/MeshDescription.h
+++ b/Engine/Source/Runtime/MeshDescription/Public/MeshDescription.h
@@ -228,15 +228,17 @@ using FEdgeArray = TMeshElementArray<FMeshEdge, FEdgeID>;
 using FPolygonArray = TMeshElementArray<FMeshPolygon, FPolygonID>;
 using FPolygonGroupArray = TMeshElementArray<FMeshPolygonGroup, FPolygonGroupID>;
 
-
+#define MESHDESCRIPTION_VER TEXT("C173EAB944C94C608318FAB504BD0F81")
 
 UCLASS()
 class MESHDESCRIPTION_API UMeshDescription : public UObject
 {
 public:
-
 	GENERATED_BODY()
 
+	//Get the current meshdescription version
+	static FString GetMeshDescriptionVersion() { return MESHDESCRIPTION_VER; }
+
 	UMeshDescription();
 
 	// UObject interface
