Fix automation build, the automation guid add for the staticmesh build ddc key was not set when executing not in debug. It is now always set.

diff --git a/Engine/Plugins/Tests/EditorTests/Source/EditorTests/Private/UnrealEd/MeshDescriptionTests.cpp b/Engine/Plugins/Tests/EditorTests/Source/EditorTests/Private/UnrealEd/MeshDescriptionTests.cpp
--- a/Engine/Plugins/Tests/EditorTests/Source/EditorTests/Private/UnrealEd/MeshDescriptionTests.cpp
+++ b/Engine/Plugins/Tests/EditorTests/Source/EditorTests/Private/UnrealEd/MeshDescriptionTests.cpp
@@ -482,7 +482,7 @@ bool FMeshDescriptionTest::ConversionTest(FAutomationTestExecutionInfo& Executio
 		FString FullAssetName = TEXT("/Game/Tests/MeshDescription/") + AssetName + TEXT(".") + AssetName;
 		UStaticMesh* AssetMesh = LoadObject<UStaticMesh>(nullptr, *FullAssetName, nullptr, LOAD_None, nullptr);
 
-#if UE_BUILD_DEBUG
+#if WITH_EDITOR
 		AssetMesh->BuildCacheAutomationTestGuid = FGuid::NewGuid();
 #endif
 
diff --git a/Engine/Source/Editor/UnrealEd/Private/Fbx/FbxStaticMeshImport.cpp b/Engine/Source/Editor/UnrealEd/Private/Fbx/FbxStaticMeshImport.cpp
--- a/Engine/Source/Editor/UnrealEd/Private/Fbx/FbxStaticMeshImport.cpp
+++ b/Engine/Source/Editor/UnrealEd/Private/Fbx/FbxStaticMeshImport.cpp
@@ -1722,13 +1722,11 @@ void UnFbx::FFbxImporter::PostImportStaticMesh(UStaticMesh* StaticMesh, TArray<F
 	// Build the staticmesh, we move the build here because we want to avoid building the staticmesh for every LOD
 	// when we import the mesh.
 	TArray<FText> BuildErrors;
-#if UE_BUILD_DEBUG
 	if (GIsAutomationTesting)
 	{
 		//Generate a random GUID to be sure it rebuild the asset
 		StaticMesh->BuildCacheAutomationTestGuid = FGuid::NewGuid();
 	}
-#endif
 
 	if (GIsAutomationTesting)
 	{
diff --git a/Engine/Source/Runtime/Engine/Private/StaticMesh.cpp b/Engine/Source/Runtime/Engine/Private/StaticMesh.cpp
--- a/Engine/Source/Runtime/Engine/Private/StaticMesh.cpp
+++ b/Engine/Source/Runtime/Engine/Private/StaticMesh.cpp
@@ -1374,13 +1374,11 @@ static FString BuildStaticMeshDerivedDataKey(UStaticMesh* Mesh, const FStaticMes
 		KeySuffix += Lex::ToString(Mesh->LightmapUVVersion);
 	}
 #if WITH_EDITOR
-#if UE_BUILD_DEBUG
 	if (GIsAutomationTesting && Mesh->BuildCacheAutomationTestGuid.IsValid())
 	{
 		//If we are in automation testing and the BuildCacheAutomationTestGuid was set
 		KeySuffix += Mesh->BuildCacheAutomationTestGuid.ToString(EGuidFormats::Digits);
 	}
-#endif
 #endif
 
 	int32 NumLODs = Mesh->SourceModels.Num();
