Fix automation test on AMD CPU/GPU. Put a threshold when testing float value, this is a tentaive fix since I do not have a AMD CPU and a AMD GPU.

diff --git a/Engine/Plugins/Tests/EditorTests/Source/EditorTests/Private/UnrealEd/MeshDescriptionTests.cpp b/Engine/Plugins/Tests/EditorTests/Source/EditorTests/Private/UnrealEd/MeshDescriptionTests.cpp
--- a/Engine/Plugins/Tests/EditorTests/Source/EditorTests/Private/UnrealEd/MeshDescriptionTests.cpp
+++ b/Engine/Plugins/Tests/EditorTests/Source/EditorTests/Private/UnrealEd/MeshDescriptionTests.cpp
@@ -142,7 +142,7 @@ bool FMeshDescriptionTest::Execute(FAutomationTestExecutionInfo& ExecutionInfo)
 }
 
 template<typename T>
-void StructureArrayCompare(const FString& ConversionName, const FString& AssetName, FAutomationTestExecutionInfo& ExecutionInfo, bool& bIsSame, const FString& VectorArrayName, const TArray<T>& ReferenceArray, const TArray<T>& ResultArray)
+void StructureArrayCompareFullPrecision(const FString& ConversionName, const FString& AssetName, FAutomationTestExecutionInfo& ExecutionInfo, bool& bIsSame, const FString& VectorArrayName, const TArray<T>& ReferenceArray, const TArray<T>& ResultArray)
 {
 	if (ReferenceArray.Num() != ResultArray.Num())
 	{
@@ -176,6 +176,41 @@ void StructureArrayCompare(const FString& ConversionName, const FString& AssetNa
 	}
 }
 
+template<typename T>
+void StructureArrayCompare(const FString& ConversionName, const FString& AssetName, FAutomationTestExecutionInfo& ExecutionInfo, bool& bIsSame, const FString& VectorArrayName, const TArray<T>& ReferenceArray, const TArray<T>& ResultArray)
+{
+	if (ReferenceArray.Num() != ResultArray.Num())
+	{
+		ExecutionInfo.AddEvent(FAutomationEvent(EAutomationEventType::Error, FString::Printf(TEXT("The %s conversion %s is not lossless, %s count is different. %s count expected [%d] result [%d]"),
+			*AssetName,
+			*ConversionName,
+			*VectorArrayName,
+			*VectorArrayName,
+			ReferenceArray.Num(),
+			ResultArray.Num())));
+		bIsSame = false;
+	}
+	else
+	{
+		for (int32 VertexIndex = 0; VertexIndex < ReferenceArray.Num(); ++VertexIndex)
+		{
+			if (!ReferenceArray[VertexIndex].Equals(ResultArray[VertexIndex]))
+			{
+				ExecutionInfo.AddEvent(FAutomationEvent(EAutomationEventType::Error, FString::Printf(TEXT("The %s conversion %s is not lossless, %s array is different. Array index [%d] expected %s [%s] result [%s]"),
+					*AssetName,
+					*ConversionName,
+					*VectorArrayName,
+					VertexIndex,
+					*VectorArrayName,
+					*ReferenceArray[VertexIndex].ToString(),
+					*ResultArray[VertexIndex].ToString())));
+				bIsSame = false;
+				break;
+			}
+		}
+	}
+}
+
 template<typename T>
 void NumberArrayCompare(const FString& ConversionName, const FString& AssetName, FAutomationTestExecutionInfo& ExecutionInfo, bool& bIsSame, const FString& VectorArrayName, const TArray<T>& ReferenceArray, const TArray<T>& ResultArray)
 {
@@ -231,8 +266,8 @@ bool FMeshDescriptionTest::CompareRawMesh(const FString& AssetName, FAutomationT
 	//BiNormal
 	StructureArrayCompare<FVector>(ConversionName, AssetName, ExecutionInfo, bAllSame, TEXT("vertex instance binormals"), ReferenceRawMesh.WedgeTangentY, ResultRawMesh.WedgeTangentY);
 
-	//Colors
-	StructureArrayCompare<FColor>(ConversionName, AssetName, ExecutionInfo, bAllSame, TEXT("vertex instance colors"), ReferenceRawMesh.WedgeColors, ResultRawMesh.WedgeColors);
+	//Colors --- FColor do not have Equals, so let use the full precision (FColor use integer anyway)
+	StructureArrayCompareFullPrecision<FColor>(ConversionName, AssetName, ExecutionInfo, bAllSame, TEXT("vertex instance colors"), ReferenceRawMesh.WedgeColors, ResultRawMesh.WedgeColors);
 
 	//Uvs
 	for (int32 UVIndex = 0; UVIndex < MAX_MESH_TEXTURE_COORDS; ++UVIndex)
