Fixed a typo in MeshDescriptionTests.

diff --git a/Engine/Plugins/Tests/EditorTests/Source/EditorTests/Private/UnrealEd/MeshDescriptionTests.cpp b/Engine/Plugins/Tests/EditorTests/Source/EditorTests/Private/UnrealEd/MeshDescriptionTests.cpp
--- a/Engine/Plugins/Tests/EditorTests/Source/EditorTests/Private/UnrealEd/MeshDescriptionTests.cpp
+++ b/Engine/Plugins/Tests/EditorTests/Source/EditorTests/Private/UnrealEd/MeshDescriptionTests.cpp
@@ -308,7 +308,7 @@ void MeshDescriptionAttributeArrayCompare(const FString& ConversionName, const F
 			ResultArray.GetNumElements())));
 		bIsSame = false;
 	}
-	else if (ReferenceArray.GetNumIndices() != ReferenceArray.GetNumIndices())
+	else if (ReferenceArray.GetNumIndices() != ResultArray.GetNumIndices())
 	{
 		ExecutionInfo.AddEvent(FAutomationEvent(EAutomationEventType::Error, FString::Printf(TEXT("The %s conversion %s is not lossless, %s channel count is different. %s channel count expected [%d] result [%d]"),
 			*AssetName,
@@ -318,7 +318,6 @@ void MeshDescriptionAttributeArrayCompare(const FString& ConversionName, const F
 			ReferenceArray.GetNumIndices(),
 			ResultArray.GetNumIndices())));
 		bIsSame = false;
-
 	}
 	else
 	{
@@ -354,114 +353,6 @@ void MeshDescriptionAttributeArrayCompare(const FString& ConversionName, const F
 	}
 }
 
-
-/*
-template<typename T, typename U, typename V>
-void MeshDescriptionStructureArrayCompare(const FString& ConversionName, const FString& AssetName, FAutomationTestExecutionInfo& ExecutionInfo, bool& bIsSame, const U& ElementIterator, const FString& VectorArrayName, const T& ReferenceArray, const T& ResultArray)
-{
-	if (ReferenceArray.Num() != ResultArray.Num())
-	{
-		ExecutionInfo.AddEvent(FAutomationEvent(EAutomationEventType::Error, FString::Printf(TEXT("The %s conversion %s is not lossless, %s count is different. %s count expected [%d] result [%d]"),
-			*AssetName,
-			*ConversionName,
-			*VectorArrayName,
-			*VectorArrayName,
-			ReferenceArray.Num(),
-			ResultArray.Num())));
-		bIsSame = false;
-	}
-	else
-	{
-		for (V ElementID : ElementIterator.GetElementIDs())
-		{
-			if (ReferenceArray[ElementID] != ResultArray[ElementID])
-			{
-				ExecutionInfo.AddEvent(FAutomationEvent(EAutomationEventType::Error, FString::Printf(TEXT("The %s conversion %s is not lossless, %s array is different. Array index [%d] expected %s [%s] result [%s]"),
-					*AssetName,
-					*ConversionName,
-					*VectorArrayName,
-					ElementID.GetValue(),
-					*VectorArrayName,
-					*ReferenceArray[ElementID].ToString(),
-					*ResultArray[ElementID].ToString())));
-				bIsSame = false;
-				break;
-			}
-		}
-	}
-}
-
-template<typename T, typename U, typename V>
-void MeshDescriptionAttributeArrayCompare(const FString& ConversionName, const FString& AssetName, FAutomationTestExecutionInfo& ExecutionInfo, bool& bIsSame, const U& ElementIterator, const FString& VectorArrayName, const T& ReferenceArray, const T& ResultArray)
-{
-	if (ReferenceArray.Num() != ResultArray.Num())
-	{
-		ExecutionInfo.AddEvent(FAutomationEvent(EAutomationEventType::Error, FString::Printf(TEXT("The %s conversion %s is not lossless, %s count is different. %s count expected [%d] result [%d]"),
-			*AssetName,
-			*ConversionName,
-			*VectorArrayName,
-			*VectorArrayName,
-			ReferenceArray.Num(),
-			ResultArray.Num())));
-		bIsSame = false;
-	}
-	else
-	{
-		for (V ElementID : ElementIterator.GetElementIDs())
-		{
-			if (ReferenceArray[ElementID] != ResultArray[ElementID])
-			{
-				ExecutionInfo.AddEvent(FAutomationEvent(EAutomationEventType::Error, FString::Printf(TEXT("The %s conversion %s is not lossless, %s array is different. Array index [%d] expected %s [%d] result [%d]"),
-					*AssetName,
-					*ConversionName,
-					*VectorArrayName,
-					ElementID.GetValue(),
-					*VectorArrayName,
-					ReferenceArray[ElementID],
-					ResultArray[ElementID])));
-				bIsSame = false;
-				break;
-			}
-		}
-	}
-}
-
-template<typename U, typename V>
-void MeshDescriptionNumberFNameCompare(const FString& ConversionName, const FString& AssetName, FAutomationTestExecutionInfo& ExecutionInfo, bool& bIsSame, const U& ElementIterator, const FString& VectorArrayName, const TPolygonGroupAttributeArray<FName>& ReferenceArray, const TPolygonGroupAttributeArray<FName>& ResultArray)
-{
-	if (ReferenceArray.Num() != ResultArray.Num())
-	{
-		ExecutionInfo.AddEvent(FAutomationEvent(EAutomationEventType::Error, FString::Printf(TEXT("The %s conversion %s is not lossless, %s count is different. %s count expected [%d] result [%d]"),
-			*AssetName,
-			*ConversionName,
-			*VectorArrayName,
-			*VectorArrayName,
-			ReferenceArray.Num(),
-			ResultArray.Num())));
-		bIsSame = false;
-	}
-	else
-	{
-		for (V ElementID : ElementIterator.GetElementIDs())
-		{
-			if (ReferenceArray[ElementID] != ResultArray[ElementID])
-			{
-				ExecutionInfo.AddEvent(FAutomationEvent(EAutomationEventType::Error, FString::Printf(TEXT("The %s conversion %s is not lossless, %s array is different. Array index [%d] expected %s [%s] result [%s]"),
-					*AssetName,
-					*ConversionName,
-					*VectorArrayName,
-					ElementID.GetValue(),
-					*VectorArrayName,
-					*ReferenceArray[ElementID].ToString(),
-					*ResultArray[ElementID].ToString())));
-				bIsSame = false;
-				break;
-			}
-		}
-	}
-}
-*/
-
 bool FMeshDescriptionTest::CompareMeshDescription(const FString& AssetName, FAutomationTestExecutionInfo& ExecutionInfo, const FMeshDescription& ReferenceMeshDescription, const FMeshDescription& MeshDescription) const
 {
 	//////////////////////////////////////////////////////////////////////////
