Refactoring the majority of AssetEditorManager cases to AssetEditorSubsystem

diff --git a/Engine/Plugins/2D/Paper2D/Source/Paper2DEditor/Private/FlipbookEditor/FlipbookEditor.cpp b/Engine/Plugins/2D/Paper2D/Source/Paper2DEditor/Private/FlipbookEditor/FlipbookEditor.cpp
--- a/Engine/Plugins/2D/Paper2D/Source/Paper2DEditor/Private/FlipbookEditor/FlipbookEditor.cpp
+++ b/Engine/Plugins/2D/Paper2D/Source/Paper2DEditor/Private/FlipbookEditor/FlipbookEditor.cpp
@@ -20,6 +20,7 @@
 #include "FlipbookEditor/SFlipbookTimeline.h"
 #include "Widgets/Docking/SDockTab.h"
 #include "Framework/Commands/GenericCommands.h"
+#include "Subsystems/AssetEditorSubsystem.h"
 
 #define LOCTEXT_NAMESPACE "FlipbookEditor"
 
@@ -313,7 +314,7 @@ void FFlipbookEditor::UnregisterTabSpawners(const TSharedRef<class FTabManager>&
 
 void FFlipbookEditor::InitFlipbookEditor(const EToolkitMode::Type Mode, const TSharedPtr< class IToolkitHost >& InitToolkitHost, class UPaperFlipbook* InitFlipbook)
 {
-	FAssetEditorManager::Get().CloseOtherEditors(InitFlipbook, this);
+	GEditor->GetEditorSubsystem<UAssetEditorSubsystem>()->CloseOtherEditors(InitFlipbook, this);
 	FlipbookBeingEdited = InitFlipbook;
 	CurrentSelectedKeyframe = INDEX_NONE;
 
diff --git a/Engine/Plugins/2D/Paper2D/Source/Paper2DEditor/Private/FlipbookEditor/STimelineTrack.cpp b/Engine/Plugins/2D/Paper2D/Source/Paper2DEditor/Private/FlipbookEditor/STimelineTrack.cpp
--- a/Engine/Plugins/2D/Paper2D/Source/Paper2DEditor/Private/FlipbookEditor/STimelineTrack.cpp
+++ b/Engine/Plugins/2D/Paper2D/Source/Paper2DEditor/Private/FlipbookEditor/STimelineTrack.cpp
@@ -20,7 +20,8 @@
 #include "Editor.h"
 #include "PropertyCustomizationHelpers.h"
 #include "PaperSprite.h"
-#include "Toolkits/AssetEditorManager.h"
+
+#include "Subsystems/AssetEditorSubsystem.h"
 
 #define LOCTEXT_NAMESPACE "FlipbookEditor"
 
@@ -337,7 +338,7 @@ void SFlipbookKeyframeWidget::EditKeyFrame()
 	{
 		if (KeyFrame->Sprite != nullptr)
 		{
-			FAssetEditorManager::Get().OpenEditorForAsset(KeyFrame->Sprite);
+			GEditor->GetEditorSubsystem<UAssetEditorSubsystem>()->OpenEditorForAsset(KeyFrame->Sprite);
 		}
 	}
 }
diff --git a/Engine/Plugins/2D/Paper2D/Source/Paper2DEditor/Private/TileMapEditing/PaperTileMapDetailsCustomization.cpp b/Engine/Plugins/2D/Paper2D/Source/Paper2DEditor/Private/TileMapEditing/PaperTileMapDetailsCustomization.cpp
--- a/Engine/Plugins/2D/Paper2D/Source/Paper2DEditor/Private/TileMapEditing/PaperTileMapDetailsCustomization.cpp
+++ b/Engine/Plugins/2D/Paper2D/Source/Paper2DEditor/Private/TileMapEditing/PaperTileMapDetailsCustomization.cpp
@@ -19,7 +19,7 @@
 #include "DetailCategoryBuilder.h"
 #include "PaperTileSet.h"
 #include "TileMapEditing/EdModeTileMap.h"
-#include "Toolkits/AssetEditorManager.h"
+
 #include "IDetailsView.h"
 #include "IAssetTools.h"
 #include "AssetToolsModule.h"
@@ -29,6 +29,7 @@
 #include "TileMapEditing/STileLayerList.h"
 #include "ScopedTransaction.h"
 #include "IPropertyUtilities.h"
+#include "Subsystems/AssetEditorSubsystem.h"
 
 #define LOCTEXT_NAMESPACE "Paper2D"
 
@@ -319,7 +320,7 @@ FReply FPaperTileMapDetailsCustomization::EnterTileMapEditingMode()
 		}
 		else
 		{
-			FAssetEditorManager::Get().OpenEditorForAsset(TileMapComponent->TileMap);
+			GEditor->GetEditorSubsystem<UAssetEditorSubsystem>()->OpenEditorForAsset(TileMapComponent->TileMap);
 		}
 	}
 	return FReply::Handled();
diff --git a/Engine/Plugins/2D/Paper2D/Source/Paper2DEditor/Private/TileMapEditing/TileMapEditor.cpp b/Engine/Plugins/2D/Paper2D/Source/Paper2DEditor/Private/TileMapEditing/TileMapEditor.cpp
--- a/Engine/Plugins/2D/Paper2D/Source/Paper2DEditor/Private/TileMapEditing/TileMapEditor.cpp
+++ b/Engine/Plugins/2D/Paper2D/Source/Paper2DEditor/Private/TileMapEditing/TileMapEditor.cpp
@@ -15,6 +15,7 @@
 #include "SCommonEditorViewportToolbarBase.h"
 #include "TileMapEditing/STileMapEditorViewportToolbar.h"
 #include "Widgets/Docking/SDockTab.h"
+#include "Subsystems/AssetEditorSubsystem.h"
 
 #define LOCTEXT_NAMESPACE "TileMapEditor"
 
@@ -298,7 +299,7 @@ void FTileMapEditor::UnregisterTabSpawners(const TSharedRef<class FTabManager>&
 
 void FTileMapEditor::InitTileMapEditor(const EToolkitMode::Type Mode, const TSharedPtr< class IToolkitHost >& InitToolkitHost, class UPaperTileMap* InitTileMap)
 {
-	FAssetEditorManager::Get().CloseOtherEditors(InitTileMap, this);
+	GEditor->GetEditorSubsystem<UAssetEditorSubsystem>()->CloseOtherEditors(InitTileMap, this);
 	TileMapBeingEdited = InitTileMap;
 
 	FTileMapEditorCommands::Register();
diff --git a/Engine/Plugins/2D/Paper2D/Source/Paper2DEditor/Private/TileSetEditor.cpp b/Engine/Plugins/2D/Paper2D/Source/Paper2DEditor/Private/TileSetEditor.cpp
--- a/Engine/Plugins/2D/Paper2D/Source/Paper2DEditor/Private/TileSetEditor.cpp
+++ b/Engine/Plugins/2D/Paper2D/Source/Paper2DEditor/Private/TileSetEditor.cpp
@@ -14,6 +14,7 @@
 #include "TileSetEditor/TileSetDetailsCustomization.h"
 #include "PaperEditorShared/SpriteGeometryEditCommands.h"
 #include "IDetailsView.h"
+#include "Subsystems/AssetEditorSubsystem.h"
 
 #define LOCTEXT_NAMESPACE "TileSetEditor"
 
@@ -167,7 +168,7 @@ void FTileSetEditor::UnregisterTabSpawners(const TSharedRef<class FTabManager>&
 
 void FTileSetEditor::InitTileSetEditor(const EToolkitMode::Type Mode, const TSharedPtr< class IToolkitHost >& InitToolkitHost, class UPaperTileSet* InitTileSet)
 {
-	FAssetEditorManager::Get().CloseOtherEditors(InitTileSet, this);
+	GEditor->GetEditorSubsystem<UAssetEditorSubsystem>()->CloseOtherEditors(InitTileSet, this);
 	TileSetBeingEdited = InitTileSet;
 
 	TileSetViewport = SNew(STileSetSelectorViewport, InitTileSet, /*EdMode=*/ nullptr);
diff --git a/Engine/Plugins/Developer/Concert/ConcertSync/ConcertSyncClient/Source/ConcertSyncClient/Private/ConcertClientSequencerManager.cpp b/Engine/Plugins/Developer/Concert/ConcertSync/ConcertSyncClient/Source/ConcertSyncClient/Private/ConcertClientSequencerManager.cpp
--- a/Engine/Plugins/Developer/Concert/ConcertSync/ConcertSyncClient/Source/ConcertSyncClient/Private/ConcertClientSequencerManager.cpp
+++ b/Engine/Plugins/Developer/Concert/ConcertSync/ConcertSyncClient/Source/ConcertSyncClient/Private/ConcertClientSequencerManager.cpp
@@ -18,8 +18,9 @@
 #if WITH_EDITOR
 	#include "ISequencerModule.h"
 	#include "ISequencer.h"
-	#include "Toolkits/AssetEditorManager.h"
+	
 #endif
+#include "Subsystems/AssetEditorSubsystem.h"
 
 DEFINE_LOG_CATEGORY_STATIC(LogConcertSequencerSync, Warning, Log)
 
@@ -315,7 +316,7 @@ void FConcertClientSequencerManager::ApplyTransportOpenEvent(const FString& Sequ
 	TGuardValue<bool> ReentrancyGuard(bRespondingToTransportEvent, true);
 	if (IsSequencerRemoteOpenEnabled())
 	{
-		FAssetEditorManager::Get().OpenEditorForAsset(SequenceObjectPath);
+		GEditor->GetEditorSubsystem<UAssetEditorSubsystem>()->OpenEditorForAsset(SequenceObjectPath);
 	}
 }
 
diff --git a/Engine/Plugins/Developer/Concert/ConcertSync/ConcertSyncClient/Source/ConcertSyncClient/Private/ConcertSyncClientUtil.cpp b/Engine/Plugins/Developer/Concert/ConcertSync/ConcertSyncClient/Source/ConcertSyncClient/Private/ConcertSyncClientUtil.cpp
--- a/Engine/Plugins/Developer/Concert/ConcertSync/ConcertSyncClient/Source/ConcertSyncClient/Private/ConcertSyncClientUtil.cpp
+++ b/Engine/Plugins/Developer/Concert/ConcertSync/ConcertSyncClient/Source/ConcertSyncClient/Private/ConcertSyncClientUtil.cpp
@@ -34,10 +34,11 @@
 	#include "Editor/UnrealEdEngine.h"
 	#include "PackageTools.h"
 	#include "ObjectTools.h"
-	#include "Toolkits/AssetEditorManager.h"
+	
 	#include "GameMapsSettings.h"
 	#include "FileHelpers.h"
 #endif
+#include "Subsystems/AssetEditorSubsystem.h"
 
 namespace ConcertSyncClientUtil
 {
@@ -357,7 +358,7 @@ void PurgePackages(TArrayView<const FName> InPackageNames)
 #if WITH_EDITOR
 		if (InObject->IsAsset())
 		{
-			FAssetEditorManager::Get().CloseAllEditorsForAsset(InObject);
+			GEditor->GetEditorSubsystem<UAssetEditorSubsystem>()->CloseAllEditorsForAsset(InObject);
 		}
 #endif
 		if (InObject->IsRooted())
diff --git a/Engine/Plugins/Editor/AssetManagerEditor/Source/AssetManagerEditor/Private/AssetManagerEditorModule.cpp b/Engine/Plugins/Editor/AssetManagerEditor/Source/AssetManagerEditor/Private/AssetManagerEditorModule.cpp
--- a/Engine/Plugins/Editor/AssetManagerEditor/Source/AssetManagerEditor/Private/AssetManagerEditorModule.cpp
+++ b/Engine/Plugins/Editor/AssetManagerEditor/Source/AssetManagerEditor/Private/AssetManagerEditorModule.cpp
@@ -33,7 +33,7 @@
 #include "Widgets/SToolTip.h"
 #include "PropertyCustomizationHelpers.h"
 #include "Toolkits/AssetEditorToolkit.h"
-#include "Toolkits/AssetEditorManager.h"
+
 #include "LevelEditor.h"
 #include "GraphEditorModule.h"
 #include "AssetData.h"
@@ -62,6 +62,7 @@
 #include "UObject/ConstructorHelpers.h"
 #include "Misc/HotReloadInterface.h"
 #include "Misc/ScopedSlowTask.h"
+#include "Subsystems/AssetEditorSubsystem.h"
 
 
 #define LOCTEXT_NAMESPACE "AssetManagerEditor"
@@ -965,14 +966,14 @@ void FAssetManagerEditorModule::OnEditAssetIdentifiers(TArray<FAssetIdentifier>
 	{
 		FScopedSlowTask SlowTask(0, LOCTEXT("LoadingSelectedObject", "Editing assets..."));
 		SlowTask.MakeDialogDelayed(.1f);
-		FAssetEditorManager& EditorManager = FAssetEditorManager::Get();
+		UAssetEditorSubsystem* AssetEditorSubsystem = GEditor->GetEditorSubsystem<UAssetEditorSubsystem>();
 
 		for (const FAssetData& AssetData : AssetsToLoad)
 		{
 			UObject* EditObject = AssetData.GetAsset();
 			if (EditObject)
 			{
-				EditorManager.OpenEditorForAsset(EditObject);
+				AssetEditorSubsystem->OpenEditorForAsset(EditObject);
 			}
 		}
 	}
diff --git a/Engine/Plugins/Editor/AssetManagerEditor/Source/AssetManagerEditor/Private/SAssetAuditBrowser.cpp b/Engine/Plugins/Editor/AssetManagerEditor/Source/AssetManagerEditor/Private/SAssetAuditBrowser.cpp
--- a/Engine/Plugins/Editor/AssetManagerEditor/Source/AssetManagerEditor/Private/SAssetAuditBrowser.cpp
+++ b/Engine/Plugins/Editor/AssetManagerEditor/Source/AssetManagerEditor/Private/SAssetAuditBrowser.cpp
@@ -9,7 +9,7 @@
 #include "Widgets/Text/STextBlock.h"
 #include "Widgets/Input/SMenuAnchor.h"
 #include "Framework/MultiBox/MultiBoxBuilder.h"
-#include "Toolkits/AssetEditorManager.h"
+
 #include "EditorStyleSet.h"
 #include "Widgets/Layout/SBox.h"
 #include "Widgets/Input/SButton.h"
@@ -32,6 +32,7 @@
 #include "Framework/Application/SlateApplication.h"
 #include "DragAndDrop/AssetDragDropOp.h"
 #include "Blueprint/BlueprintSupport.h"
+#include "Subsystems/AssetEditorSubsystem.h"
 
 #define LOCTEXT_NAMESPACE "AssetManagementBrowser"
 
@@ -132,7 +133,7 @@ void SAssetAuditBrowser::EditSelectedAssets() const
 		AssetNames.Add(AssetData.ObjectPath);
 	}
 
-	FAssetEditorManager::Get().OpenEditorsForAssets(AssetNames);
+	GEditor->GetEditorSubsystem<UAssetEditorSubsystem>()->OpenEditorsForAssets(AssetNames);
 }
 
 void SAssetAuditBrowser::OnRequestOpenAsset(const FAssetData& AssetData) const
@@ -140,7 +141,7 @@ void SAssetAuditBrowser::OnRequestOpenAsset(const FAssetData& AssetData) const
 	TArray<FName> AssetNames;
 	AssetNames.Add(AssetData.ObjectPath);
 
-	FAssetEditorManager::Get().OpenEditorsForAssets(AssetNames);
+	GEditor->GetEditorSubsystem<UAssetEditorSubsystem>()->OpenEditorsForAssets(AssetNames);
 }
 
 void SAssetAuditBrowser::SaveSelectedAssets() const
diff --git a/Engine/Plugins/Editor/AssetManagerEditor/Source/AssetManagerEditor/Private/SSizeMap.cpp b/Engine/Plugins/Editor/AssetManagerEditor/Source/AssetManagerEditor/Private/SSizeMap.cpp
--- a/Engine/Plugins/Editor/AssetManagerEditor/Source/AssetManagerEditor/Private/SSizeMap.cpp
+++ b/Engine/Plugins/Editor/AssetManagerEditor/Source/AssetManagerEditor/Private/SSizeMap.cpp
@@ -12,7 +12,7 @@
 #include "Widgets/Input/SButton.h"
 #include "Framework/Commands/UIAction.h"
 #include "Framework/Commands/UICommandList.h"
-#include "Toolkits/AssetEditorManager.h"
+
 #include "Toolkits/GlobalEditorCommonCommands.h"
 #include "Framework/MultiBox/MultiBoxBuilder.h"
 #include "AssetManagerEditorModule.h"
@@ -22,6 +22,7 @@
 #include "ICollectionManager.h"
 #include "Framework/Application/SlateApplication.h"
 #include "Misc/ScopedSlowTask.h"
+#include "Subsystems/AssetEditorSubsystem.h"
 
 #define LOCTEXT_NAMESPACE "SizeMap"
 
@@ -963,7 +964,7 @@ void SSizeMap::EditSelectedAssets() const
 		TArray<FName> AssetNames;
 		AssetNames.Add(NodeSizeMapData->AssetData.ObjectPath);
 
-		FAssetEditorManager::Get().OpenEditorsForAssets(AssetNames);
+		GEditor->GetEditorSubsystem<UAssetEditorSubsystem>()->OpenEditorsForAssets(AssetNames);
 	}
 }
 	
diff --git a/Engine/Plugins/Editor/GameplayTagsEditor/Source/GameplayTagsEditor/Private/GameplayTagReferenceHelperDetails.cpp b/Engine/Plugins/Editor/GameplayTagsEditor/Source/GameplayTagsEditor/Private/GameplayTagReferenceHelperDetails.cpp
--- a/Engine/Plugins/Editor/GameplayTagsEditor/Source/GameplayTagsEditor/Private/GameplayTagReferenceHelperDetails.cpp
+++ b/Engine/Plugins/Editor/GameplayTagsEditor/Source/GameplayTagsEditor/Private/GameplayTagReferenceHelperDetails.cpp
@@ -3,7 +3,7 @@
 #include "GameplayTagReferenceHelperDetails.h"
 #include "AssetRegistryModule.h"
 #include "Widgets/Input/SHyperlink.h"
-#include "Toolkits/AssetEditorManager.h"
+
 #include "ObjectTools.h"
 #include "GameplayTagContainer.h"
 #include "UObject/UObjectHash.h"
@@ -11,6 +11,7 @@
 #include "GameplayTagsManager.h"
 #include "SGameplayTagWidget.h"
 #include "IDetailChildrenBuilder.h"
+#include "Subsystems/AssetEditorSubsystem.h"
 
 #define LOCTEXT_NAMESPACE "GameplayTagReferenceHelperDetails"
 
@@ -172,7 +173,7 @@ TSharedRef<ITableRow> FGameplayTagReferenceHelperDetails::OnGenerateWidgetForGam
 				{
 					if (ObjectTools::IsObjectBrowsable(Obj))
 					{
-						FAssetEditorManager::Get().OpenEditorForAsset(Obj);
+						GEditor->GetEditorSubsystem<UAssetEditorSubsystem>()->OpenEditorForAsset(Obj);
 					}
 				});
 			}
diff --git a/Engine/Plugins/Editor/GameplayTagsEditor/Source/GameplayTagsEditor/Private/SGameplayTagWidget.cpp b/Engine/Plugins/Editor/GameplayTagsEditor/Source/GameplayTagsEditor/Private/SGameplayTagWidget.cpp
--- a/Engine/Plugins/Editor/GameplayTagsEditor/Source/GameplayTagsEditor/Private/SGameplayTagWidget.cpp
+++ b/Engine/Plugins/Editor/GameplayTagsEditor/Source/GameplayTagsEditor/Private/SGameplayTagWidget.cpp
@@ -16,7 +16,7 @@
 #include "Widgets/Input/SSearchBox.h"
 #include "GameplayTagsEditorModule.h"
 #include "Widgets/Layout/SScaleBox.h"
-#include "Toolkits/AssetEditorManager.h"
+
 #include "AssetToolsModule.h"
 #include "Widgets/Input/SHyperlink.h"
 #include "Widgets/Notifications/SNotificationList.h"
diff --git a/Engine/Plugins/Enterprise/DatasmithContent/Source/DatasmithContentEditor/Private/DatasmithSceneNewFactory.cpp b/Engine/Plugins/Enterprise/DatasmithContent/Source/DatasmithContentEditor/Private/DatasmithSceneNewFactory.cpp
--- a/Engine/Plugins/Enterprise/DatasmithContent/Source/DatasmithContentEditor/Private/DatasmithSceneNewFactory.cpp
+++ b/Engine/Plugins/Enterprise/DatasmithContent/Source/DatasmithContentEditor/Private/DatasmithSceneNewFactory.cpp
@@ -13,7 +13,7 @@
 #include "Misc/MessageDialog.h"
 #include "Misc/Paths.h"
 #include "Modules/ModuleManager.h"
-#include "Toolkits/AssetEditorManager.h"
+
 #include "UObject/SoftObjectPtr.h"
 #include "UObject/SoftObjectPath.h"
 
diff --git a/Engine/Plugins/Experimental/AlembicImporter/Source/AlembicImporter/Private/AlembicImportFactory.cpp b/Engine/Plugins/Experimental/AlembicImporter/Source/AlembicImporter/Private/AlembicImportFactory.cpp
--- a/Engine/Plugins/Experimental/AlembicImporter/Source/AlembicImporter/Private/AlembicImportFactory.cpp
+++ b/Engine/Plugins/Experimental/AlembicImporter/Source/AlembicImporter/Private/AlembicImportFactory.cpp
@@ -7,7 +7,7 @@
 #include "EditorFramework/AssetImportData.h"
 #include "HAL/FileManager.h"
 #include "Framework/Application/SlateApplication.h"
-#include "Toolkits/AssetEditorManager.h"
+
 
 #include "Interfaces/IMainFrameModule.h"
 
@@ -20,6 +20,7 @@
 #include "AbcAssetImportData.h"
 
 #include "GeometryCache.h"
+#include "Subsystems/AssetEditorSubsystem.h"
 
 #define LOCTEXT_NAMESPACE "AlembicImportFactory"
 
@@ -380,7 +381,7 @@ EReimportResult::Type UAlembicImportFactory::Reimport(UObject* Obj)
 		CurrentFilename = Mesh->AssetImportData->GetFirstFilename();
 
 		// Close possible open editors using this asset	
-		FAssetEditorManager::Get().CloseAllEditorsForAsset(Mesh);
+		GEditor->GetEditorSubsystem<UAssetEditorSubsystem>()->CloseAllEditorsForAsset(Mesh);
 		
 		EReimportResult::Type Result = ReimportStaticMesh(Mesh);
 		FAbcImportLogger::OutputMessages(PageName);
@@ -409,7 +410,7 @@ EReimportResult::Type UAlembicImportFactory::Reimport(UObject* Obj)
 		}
 
 		// Close possible open editors using this asset	
-		FAssetEditorManager::Get().CloseAllEditorsForAsset(GeometryCache);
+		GEditor->GetEditorSubsystem<UAssetEditorSubsystem>()->CloseAllEditorsForAsset(GeometryCache);
 
 		FAbcImportLogger::OutputMessages(PageName);
 		return Result;
@@ -436,7 +437,7 @@ EReimportResult::Type UAlembicImportFactory::Reimport(UObject* Obj)
 		}
 
 		// Close possible open editors using this asset	
-		FAssetEditorManager::Get().CloseAllEditorsForAsset(SkeletalMesh);
+		GEditor->GetEditorSubsystem<UAssetEditorSubsystem>()->CloseAllEditorsForAsset(SkeletalMesh);
 
 		FAbcImportLogger::OutputMessages(PageName);
 		return Result;
@@ -467,7 +468,7 @@ EReimportResult::Type UAlembicImportFactory::Reimport(UObject* Obj)
 		}
 
 		// Close possible open editors using this asset	
-		FAssetEditorManager::Get().CloseAllEditorsForAsset(SkeletalMesh);
+		GEditor->GetEditorSubsystem<UAssetEditorSubsystem>()->CloseAllEditorsForAsset(SkeletalMesh);
 
 		EReimportResult::Type Result = ReimportSkeletalMesh(SkeletalMesh);
 
diff --git a/Engine/Plugins/Experimental/AlembicImporter/Source/AlembicLibrary/Private/AbcImporter.cpp b/Engine/Plugins/Experimental/AlembicImporter/Source/AlembicLibrary/Private/AbcImporter.cpp
--- a/Engine/Plugins/Experimental/AlembicImporter/Source/AlembicLibrary/Private/AbcImporter.cpp
+++ b/Engine/Plugins/Experimental/AlembicImporter/Source/AlembicLibrary/Private/AbcImporter.cpp
@@ -36,7 +36,7 @@ PRAGMA_DEFAULT_VISIBILITY_END
 #include "SkelImport.h"
 #include "Animation/AnimSequence.h"
 #include "Rendering/SkeletalMeshModel.h"
-#include "Toolkits/AssetEditorManager.h"
+
 
 #include "AbcImportUtilities.h"
 #include "Utils.h"
@@ -59,6 +59,8 @@ PRAGMA_DEFAULT_VISIBILITY_END
 #include "AnimationUtils.h"
 #include "ComponentReregisterContext.h"
 #include "GeometryCacheCodecV1.h"
+#include "Subsystems/AssetEditorSubsystem.h"
+#include "Editor.h"
 
 #define LOCTEXT_NAMESPACE "AbcImporter"
 
@@ -806,10 +808,10 @@ TArray<UObject*> FAbcImporter::ImportAsSkeletalMesh(UObject* InParent, EObjectFl
 		GeneratedObjects.Add(Skeleton);
 		GeneratedObjects.Add(Sequence);
 
-		FAssetEditorManager& AssetEditorManager = FAssetEditorManager::Get();
-		AssetEditorManager.CloseAllEditorsForAsset(Skeleton);
-		AssetEditorManager.CloseAllEditorsForAsset(SkeletalMesh);
-		AssetEditorManager.CloseAllEditorsForAsset(Sequence);
+		UAssetEditorSubsystem* AssetEditorSubsystem = GEditor->GetEditorSubsystem<UAssetEditorSubsystem>();
+		AssetEditorSubsystem->CloseAllEditorsForAsset(Skeleton);
+		AssetEditorSubsystem->CloseAllEditorsForAsset(SkeletalMesh);
+		AssetEditorSubsystem->CloseAllEditorsForAsset(Sequence);
 	}
 
 	if (RecreateExistingRenderStateContext)
diff --git a/Engine/Plugins/Experimental/CodeEditor/Source/CodeEditor/Private/CodeProjectEditor.cpp b/Engine/Plugins/Experimental/CodeEditor/Source/CodeEditor/Private/CodeProjectEditor.cpp
--- a/Engine/Plugins/Experimental/CodeEditor/Source/CodeEditor/Private/CodeProjectEditor.cpp
+++ b/Engine/Plugins/Experimental/CodeEditor/Source/CodeEditor/Private/CodeProjectEditor.cpp
@@ -13,6 +13,7 @@
 #include "CodeProjectEditorToolbar.h"
 #include "WorkflowOrientedApp/ApplicationMode.h"
 #include "WorkflowOrientedApp/WorkflowUObjectDocuments.h"
+#include "Subsystems/AssetEditorSubsystem.h"
 
 
 #define LOCTEXT_NAMESPACE "CodeEditor"
@@ -204,7 +205,7 @@ void FCodeProjectEditor::RegisterToolbarTab(const TSharedRef<class FTabManager>&
 
 void FCodeProjectEditor::InitCodeEditor(const EToolkitMode::Type Mode, const TSharedPtr< class IToolkitHost >& InitToolkitHost, class UCodeProject* CodeProject)
 {
-	FAssetEditorManager::Get().CloseOtherEditors(CodeProject, this);
+	GEditor->GetEditorSubsystem<UAssetEditorSubsystem>()->CloseOtherEditors(CodeProject, this);
 	CodeProjectBeingEdited = CodeProject;
 
 	TSharedPtr<FCodeProjectEditor> ThisPtr(SharedThis(this));
diff --git a/Engine/Plugins/Experimental/ControlRig/Source/ControlRigEditor/Private/ControlRigBlueprintActions.cpp b/Engine/Plugins/Experimental/ControlRig/Source/ControlRigEditor/Private/ControlRigBlueprintActions.cpp
--- a/Engine/Plugins/Experimental/ControlRig/Source/ControlRigEditor/Private/ControlRigBlueprintActions.cpp
+++ b/Engine/Plugins/Experimental/ControlRig/Source/ControlRigEditor/Private/ControlRigBlueprintActions.cpp
@@ -4,11 +4,12 @@
 #include "ControlRigBlueprintFactory.h"
 #include "ControlRigBlueprint.h"
 #include "IControlRigEditorModule.h"
-#include "Toolkits/AssetEditorManager.h"
+
 #include "Styling/SlateIconFinder.h"
 #include "Widgets/Layout/SBorder.h"
 #include "Widgets/Images/SImage.h"
 #include "EditorStyleSet.h"
+#include "Subsystems/AssetEditorSubsystem.h"
 
 #define LOCTEXT_NAMESPACE "ControlRigBlueprintActions"
 
@@ -29,7 +30,7 @@ void FControlRigBlueprintActions::OpenAssetEditor( const TArray<UObject*>& InObj
 		if (UControlRigBlueprint* ControlRigBlueprint = Cast<UControlRigBlueprint>(*ObjIt))
 		{
 			const bool bBringToFrontIfOpen = true;
-			if (IAssetEditorInstance* EditorInstance = FAssetEditorManager::Get().FindEditorForAsset(ControlRigBlueprint, bBringToFrontIfOpen))
+			if (IAssetEditorInstance* EditorInstance = GEditor->GetEditorSubsystem<UAssetEditorSubsystem>()->FindEditorForAsset(ControlRigBlueprint, bBringToFrontIfOpen))
 			{
 				EditorInstance->FocusWindow(ControlRigBlueprint);
 			}
diff --git a/Engine/Plugins/Experimental/ControlRig/Source/ControlRigEditor/Private/ControlRigEditorModule.cpp b/Engine/Plugins/Experimental/ControlRig/Source/ControlRigEditor/Private/ControlRigEditorModule.cpp
--- a/Engine/Plugins/Experimental/ControlRig/Source/ControlRigEditor/Private/ControlRigEditorModule.cpp
+++ b/Engine/Plugins/Experimental/ControlRig/Source/ControlRigEditor/Private/ControlRigEditorModule.cpp
@@ -73,6 +73,7 @@
 #include "Editor/SControlRigProfilingView.h"
 #include "WorkspaceMenuStructure.h"
 #include "WorkspaceMenuStructureModule.h"
+#include "Subsystems/AssetEditorSubsystem.h"
 
 #define LOCTEXT_NAMESPACE "ControlRigEditorModule"
 
@@ -158,7 +159,7 @@ void FControlRigEditorModule::StartupModule()
 	SequencerModule.GetToolBarExtensibilityManager()->AddExtender(SequencerToolbarExtender);
 
 	// Register for assets being opened
-	AssetEditorOpenedHandle = FAssetEditorManager::Get().OnAssetEditorOpened().AddRaw(this, &FControlRigEditorModule::HandleAssetEditorOpened);
+	AssetEditorOpenedHandle = GEditor->GetEditorSubsystem<UAssetEditorSubsystem>()->OnAssetEditorOpened().AddRaw(this, &FControlRigEditorModule::HandleAssetEditorOpened);
 
 	// Register level sequence spawner
 	ILevelSequenceModule& LevelSequenceModule = FModuleManager::LoadModuleChecked<ILevelSequenceModule>("LevelSequence");
@@ -355,7 +356,10 @@ void FControlRigEditorModule::ShutdownModule()
 		TrajectoryMaterial->RemoveFromRoot();
 	}
 
-	FAssetEditorManager::Get().OnAssetEditorOpened().Remove(AssetEditorOpenedHandle);
+	if (UAssetEditorSubsystem* AssetEditorSubsystem = GEditor->GetEditorSubsystem<UAssetEditorSubsystem>())
+	{
+		AssetEditorSubsystem->OnAssetEditorOpened().Remove(AssetEditorOpenedHandle);
+	}
 
 	FEditorModeRegistry::Get().UnregisterMode(FControlRigEditorEditMode::ModeName);
 	FEditorModeRegistry::Get().UnregisterMode(FControlRigEditMode::ModeName);
diff --git a/Engine/Plugins/Experimental/ControlRig/Source/ControlRigEditor/Private/EditMode/ControlRigEditModeSettings.cpp b/Engine/Plugins/Experimental/ControlRig/Source/ControlRigEditor/Private/EditMode/ControlRigEditModeSettings.cpp
--- a/Engine/Plugins/Experimental/ControlRig/Source/ControlRigEditor/Private/EditMode/ControlRigEditModeSettings.cpp
+++ b/Engine/Plugins/Experimental/ControlRig/Source/ControlRigEditor/Private/EditMode/ControlRigEditModeSettings.cpp
@@ -4,8 +4,9 @@
 #include "EditorModeManager.h"
 #include "ControlRigEditMode.h"
 #include "Sequencer/ControlRigSequence.h"
-#include "Toolkits/AssetEditorManager.h"
+
 #include "Components/SkeletalMeshComponent.h"
+#include "Subsystems/AssetEditorSubsystem.h"
 
 void UControlRigEditModeSettings::PreEditChange(UProperty* PropertyAboutToChange)
 {
@@ -52,7 +53,7 @@ void UControlRigEditModeSettings::PostEditChangeProperty(struct FPropertyChanged
 		{
 			if (Sequence)
 			{
-				FAssetEditorManager::Get().OpenEditorForAsset(Sequence);
+				GEditor->GetEditorSubsystem<UAssetEditorSubsystem>()->OpenEditorForAsset(Sequence);
 			}
 		}
 		else if (Property->GetFName() == GET_MEMBER_NAME_CHECKED(UControlRigEditModeSettings, bDisplayTrajectories))
diff --git a/Engine/Plugins/Experimental/ControlRig/Source/ControlRigEditor/Private/Sequencer/ControlRigSequenceExporter.cpp b/Engine/Plugins/Experimental/ControlRig/Source/ControlRigEditor/Private/Sequencer/ControlRigSequenceExporter.cpp
--- a/Engine/Plugins/Experimental/ControlRig/Source/ControlRigEditor/Private/Sequencer/ControlRigSequenceExporter.cpp
+++ b/Engine/Plugins/Experimental/ControlRig/Source/ControlRigEditor/Private/Sequencer/ControlRigSequenceExporter.cpp
@@ -19,12 +19,13 @@
 #include "Sequencer/ControlRigBindingTemplate.h"
 #include "Editor.h"
 #include "Widgets/Notifications/SNotificationList.h"
-#include "Toolkits/AssetEditorManager.h"
+
 #include "AssetRegistryModule.h"
 #include "Framework/Notifications/NotificationManager.h"
 #include "Widgets/Notifications/SNotificationList.h"
 #include "Widgets/Text/STextBlock.h"
 #include "MovieSceneTimeHelpers.h"
+#include "Subsystems/AssetEditorSubsystem.h"
 
 #define LOCTEXT_NAMESPACE "ControlRigSequenceExporter"
 
@@ -307,7 +308,7 @@ void Convert(UControlRigSequence* Sequence, UAnimSequence* AnimSequence, USkelet
 				{
 					TArray<UObject*> Assets;
 					Assets.Add(WeakAnimSequence.Get());
-					FAssetEditorManager::Get().OpenEditorForAssets(Assets);
+					GEditor->GetEditorSubsystem<UAssetEditorSubsystem>()->OpenEditorForAssets(Assets);
 				}
 			});
 			Info.HyperlinkText = FText::Format(LOCTEXT("OpenNewAnimationHyperlink", "Open {0}"), FText::FromString(AnimSequence->GetName()));
diff --git a/Engine/Plugins/Experimental/DataPrepEditor/Source/DataPrepCore/Private/DataprepCoreUtils.cpp b/Engine/Plugins/Experimental/DataPrepEditor/Source/DataPrepCore/Private/DataprepCoreUtils.cpp
--- a/Engine/Plugins/Experimental/DataPrepEditor/Source/DataPrepCore/Private/DataprepCoreUtils.cpp
+++ b/Engine/Plugins/Experimental/DataPrepEditor/Source/DataPrepCore/Private/DataprepCoreUtils.cpp
@@ -6,8 +6,9 @@
 
 #if WITH_EDITOR
 #include "ObjectTools.h"
-#include "Toolkits/AssetEditorManager.h"
+
 #endif
+#include "Subsystems/AssetEditorSubsystem.h"
 
 void FDataprepCoreUtils::PurgeObjects(TArray<UObject*> Objects)
 {
@@ -21,7 +22,7 @@ void FDataprepCoreUtils::PurgeObjects(TArray<UObject*> Objects)
 #if WITH_EDITOR
 		if ( InObject->IsAsset() )
 		{
-			FAssetEditorManager::Get().CloseAllEditorsForAsset( InObject );
+			GEditor->GetEditorSubsystem<UAssetEditorSubsystem>()->CloseAllEditorsForAsset( InObject );
 		}
 #endif
 		if ( InObject->IsRooted() )
diff --git a/Engine/Plugins/Experimental/VirtualCamera/Source/VirtualCamera/Private/LevelSequencePlaybackController.cpp b/Engine/Plugins/Experimental/VirtualCamera/Source/VirtualCamera/Private/LevelSequencePlaybackController.cpp
--- a/Engine/Plugins/Experimental/VirtualCamera/Source/VirtualCamera/Private/LevelSequencePlaybackController.cpp
+++ b/Engine/Plugins/Experimental/VirtualCamera/Source/VirtualCamera/Private/LevelSequencePlaybackController.cpp
@@ -16,8 +16,9 @@
 #include "Input/Reply.h"
 #include "Recorder/TakeRecorderBlueprintLibrary.h"
 #include "Templates/SharedPointer.h"
-#include "Toolkits/AssetEditorManager.h"
+
 #endif //WITH_EDITOR
+#include "Subsystems/AssetEditorSubsystem.h"
 
 
 ULevelSequencePlaybackController::ULevelSequencePlaybackController(const FObjectInitializer& ObjectInitializer)
@@ -307,8 +308,8 @@ bool ULevelSequencePlaybackController::SetActiveLevelSequence(ULevelSequence* In
 	if(InNewLevelSequence && !IsRunningGame())
 	{
 		const bool bDoFocusOnEditor = false;
-		FAssetEditorManager::Get().OpenEditorForAsset(InNewLevelSequence);
-		IAssetEditorInstance* AssetEditor = FAssetEditorManager::Get().FindEditorForAsset(InNewLevelSequence, bDoFocusOnEditor);
+		GEditor->GetEditorSubsystem<UAssetEditorSubsystem>()->OpenEditorForAsset(InNewLevelSequence);
+		IAssetEditorInstance* AssetEditor = GEditor->GetEditorSubsystem<UAssetEditorSubsystem>()->FindEditorForAsset(InNewLevelSequence, bDoFocusOnEditor);
 		ILevelSequenceEditorToolkit* LevelSequenceEditor = static_cast<ILevelSequenceEditorToolkit*>(AssetEditor);
 
 		WeakSequencer = LevelSequenceEditor ? LevelSequenceEditor->GetSequencer() : nullptr;
diff --git a/Engine/Plugins/FX/Niagara/Source/NiagaraEditor/Private/NiagaraMessageManager.cpp b/Engine/Plugins/FX/Niagara/Source/NiagaraEditor/Private/NiagaraMessageManager.cpp
--- a/Engine/Plugins/FX/Niagara/Source/NiagaraEditor/Private/NiagaraMessageManager.cpp
+++ b/Engine/Plugins/FX/Niagara/Source/NiagaraEditor/Private/NiagaraMessageManager.cpp
@@ -7,9 +7,10 @@
 #include "Modules/ModuleManager.h"
 #include "AssetRegistryModule.h"
 #include "IAssetRegistry.h"
-#include "Toolkits/AssetEditorManager.h"
+
 #include "NiagaraScriptToolkit.h"
 #include "../Public/NiagaraEditorModule.h"
+#include "Subsystems/AssetEditorSubsystem.h"
 
 #define LOCTEXT_NAMESPACE "NiagaraMessageManager"
 
@@ -497,7 +498,7 @@ void FNiagaraCompileEventToken::OpenScriptAssetByPathAndFocusNodeOrPinIfSet(
 	{
 		if (UNiagaraScript* ScriptAsset = Cast<UNiagaraScript>(AssetData.GetAsset()))
 		{
-			FAssetEditorManager::Get().OpenEditorForAsset(ScriptAsset, EToolkitMode::Standalone);
+			GEditor->GetEditorSubsystem<UAssetEditorSubsystem>()->OpenEditorForAsset(ScriptAsset, EToolkitMode::Standalone);
 			FNiagaraEditorModule& NiagaraEditorModule = FModuleManager::LoadModuleChecked<FNiagaraEditorModule>("NiagaraEditor");
 			if (InPinGUID.IsSet())
 			{
diff --git a/Engine/Plugins/FX/Niagara/Source/NiagaraEditor/Private/Widgets/SNiagaraScriptGraph.cpp b/Engine/Plugins/FX/Niagara/Source/NiagaraEditor/Private/Widgets/SNiagaraScriptGraph.cpp
--- a/Engine/Plugins/FX/Niagara/Source/NiagaraEditor/Private/Widgets/SNiagaraScriptGraph.cpp
+++ b/Engine/Plugins/FX/Niagara/Source/NiagaraEditor/Private/Widgets/SNiagaraScriptGraph.cpp
@@ -6,7 +6,7 @@
 #include "NiagaraGraph.h"
 #include "NiagaraNode.h"
 #include "NiagaraNodeInput.h"
-#include "Toolkits/AssetEditorManager.h"
+
 #include "GraphEditor.h"
 #include "EditorStyleSet.h"
 #include "Widgets/Layout/SBorder.h"
@@ -26,6 +26,7 @@
 #include "Layout/WidgetPath.h"
 #include "Framework/Application/SlateApplication.h"
 #include "GraphEditorActions.h"
+#include "Subsystems/AssetEditorSubsystem.h"
 
 #define LOCTEXT_NAMESPACE "NiagaraScriptGraph"
 
@@ -207,7 +208,7 @@ void SNiagaraScriptGraph::OnNodeDoubleClicked(UEdGraphNode* ClickedNode)
 		UObject* ReferencedAsset = NiagaraNode->GetReferencedAsset();
 		if (ReferencedAsset != nullptr)
 		{
-			FAssetEditorManager::Get().OpenEditorForAsset(ReferencedAsset);
+			GEditor->GetEditorSubsystem<UAssetEditorSubsystem>()->OpenEditorForAsset(ReferencedAsset);
 		}
 	}
 }
diff --git a/Engine/Plugins/FX/Niagara/Source/NiagaraEditorWidgets/Private/SNiagaraStack.cpp b/Engine/Plugins/FX/Niagara/Source/NiagaraEditorWidgets/Private/SNiagaraStack.cpp
--- a/Engine/Plugins/FX/Niagara/Source/NiagaraEditorWidgets/Private/SNiagaraStack.cpp
+++ b/Engine/Plugins/FX/Niagara/Source/NiagaraEditorWidgets/Private/SNiagaraStack.cpp
@@ -55,6 +55,7 @@
 #include "Widgets/Layout/SWrapBox.h"
 #include "Stack/SNiagaraStackSpacer.h"
 #include "ViewModels/Stack/NiagaraStackRoot.h"
+#include "Subsystems/AssetEditorSubsystem.h"
 
 /** Contains data for a socket drag and drop operation in the StackEntry node. */
 class FNiagaraStackEntryDragDropOp : public FDecoratedDragDropOp
@@ -537,7 +538,7 @@ FReply SNiagaraStack::OpenParentEmitter()
 		UNiagaraEmitter* ParentEmitter = const_cast<UNiagaraEmitter*>(StackViewModel->GetEmitterHandleViewModel()->GetEmitterViewModel()->GetParentEmitter());
 		if (ParentEmitter != nullptr)
 		{
-			FAssetEditorManager::Get().OpenEditorForAsset(ParentEmitter);
+			GEditor->GetEditorSubsystem<UAssetEditorSubsystem>()->OpenEditorForAsset(ParentEmitter);
 		}
 	}
 	return FReply::Handled();
diff --git a/Engine/Plugins/FX/Niagara/Source/NiagaraEditorWidgets/Private/Stack/SNiagaraStackFunctionInputValue.cpp b/Engine/Plugins/FX/Niagara/Source/NiagaraEditorWidgets/Private/Stack/SNiagaraStackFunctionInputValue.cpp
--- a/Engine/Plugins/FX/Niagara/Source/NiagaraEditorWidgets/Private/Stack/SNiagaraStackFunctionInputValue.cpp
+++ b/Engine/Plugins/FX/Niagara/Source/NiagaraEditorWidgets/Private/Stack/SNiagaraStackFunctionInputValue.cpp
@@ -28,6 +28,7 @@
 #include "NiagaraParameterCollection.h"
 #include "ViewModels/NiagaraSystemViewModel.h"
 #include "NiagaraSystem.h"
+#include "Subsystems/AssetEditorSubsystem.h"
 
 #define LOCTEXT_NAMESPACE "NiagaraStackFunctionInputValue"
 
@@ -425,7 +426,7 @@ FReply SNiagaraStackFunctionInputValue::DynamicInputTextDoubleClicked(const FGeo
 	UNiagaraNodeFunctionCall* DynamicInputNode = FunctionInput->GetDynamicInputNode();
 	if (DynamicInputNode->FunctionScript != nullptr && DynamicInputNode->FunctionScript->IsAsset())
 	{
-		FAssetEditorManager::Get().OpenEditorForAsset(DynamicInputNode->FunctionScript);
+		GEditor->GetEditorSubsystem<UAssetEditorSubsystem>()->OpenEditorForAsset(DynamicInputNode->FunctionScript);
 		return FReply::Handled();
 	}
 	return FReply::Unhandled();
@@ -449,11 +450,11 @@ FReply SNiagaraStackFunctionInputValue::OnLinkedInputDoubleClicked(const FGeomet
 			if (UNiagaraParameterCollectionInstance* NPCInst = FunctionInput->GetSystemViewModel()->GetSystem().GetParameterCollectionOverride(Collection))
 			{
 				//If we override this NPC then open the instance.
-				FAssetEditorManager::Get().OpenEditorForAsset(NPCInst);
+				GEditor->GetEditorSubsystem<UAssetEditorSubsystem>()->OpenEditorForAsset(NPCInst);
 			}
 			else
 			{
-				FAssetEditorManager::Get().OpenEditorForAsset(Collection); 
+				GEditor->GetEditorSubsystem<UAssetEditorSubsystem>()->OpenEditorForAsset(Collection); 
 			}
 			
 			return FReply::Handled();
diff --git a/Engine/Plugins/FX/Niagara/Source/NiagaraEditorWidgets/Private/Stack/SNiagaraStackModuleItem.cpp b/Engine/Plugins/FX/Niagara/Source/NiagaraEditorWidgets/Private/Stack/SNiagaraStackModuleItem.cpp
--- a/Engine/Plugins/FX/Niagara/Source/NiagaraEditorWidgets/Private/Stack/SNiagaraStackModuleItem.cpp
+++ b/Engine/Plugins/FX/Niagara/Source/NiagaraEditorWidgets/Private/Stack/SNiagaraStackModuleItem.cpp
@@ -20,11 +20,12 @@
 #include "Widgets/Images/SImage.h"
 #include "Framework/MultiBox/MultiBoxBuilder.h"
 #include "Framework/Application/SlateApplication.h"
-#include "Toolkits/AssetEditorManager.h"
+
 #include "SDropTarget.h"
 #include "SNiagaraStackErrorButton.h"
 #include "NiagaraEditorUtilities.h"
 #include "SGraphActionMenu.h"
+#include "Subsystems/AssetEditorSubsystem.h"
 
 #define LOCTEXT_NAMESPACE "NiagaraStackModuleItem"
 
@@ -176,7 +177,7 @@ FReply SNiagaraStackModuleItem::OnMouseButtonDoubleClick(const FGeometry& InMyGe
 	const UNiagaraNodeFunctionCall& ModuleFunctionCall = ModuleItem->GetModuleNode();
 	if (ModuleFunctionCall.FunctionScript != nullptr && ModuleFunctionCall.FunctionScript->IsAsset())
 	{
-		FAssetEditorManager::Get().OpenEditorForAsset(const_cast<UNiagaraScript*>(ModuleFunctionCall.FunctionScript));
+		GEditor->GetEditorSubsystem<UAssetEditorSubsystem>()->OpenEditorForAsset(const_cast<UNiagaraScript*>(ModuleFunctionCall.FunctionScript));
 		return FReply::Handled();
 	}
 	return FReply::Unhandled();
diff --git a/Engine/Plugins/FX/Niagara/Source/NiagaraEditorWidgets/Private/Stack/SNiagaraStackTableRow.cpp b/Engine/Plugins/FX/Niagara/Source/NiagaraEditorWidgets/Private/Stack/SNiagaraStackTableRow.cpp
--- a/Engine/Plugins/FX/Niagara/Source/NiagaraEditorWidgets/Private/Stack/SNiagaraStackTableRow.cpp
+++ b/Engine/Plugins/FX/Niagara/Source/NiagaraEditorWidgets/Private/Stack/SNiagaraStackTableRow.cpp
@@ -13,8 +13,10 @@
 #include "Widgets/Layout/SBox.h"
 #include "IContentBrowserSingleton.h"
 #include "ContentBrowserModule.h"
-#include "Toolkits/AssetEditorManager.h"
+
 #include "Framework/Application/SlateApplication.h"
+#include "Subsystems/AssetEditorSubsystem.h"
+#include "Editor.h"
 
 #define LOCTEXT_NAMESPACE "NiagaraStackTableRow"
 
@@ -441,7 +443,7 @@ void SNiagaraStackTableRow::NavigateTo(UNiagaraStackEntry* Item)
 
 void SNiagaraStackTableRow::OpenSourceAsset()
 {
-	FAssetEditorManager::Get().OpenEditorForAsset(StackEntry->GetExternalAsset());
+	GEditor->GetEditorSubsystem<UAssetEditorSubsystem>()->OpenEditorForAsset(StackEntry->GetExternalAsset());
 }
 
 void SNiagaraStackTableRow::ShowAssetInContentBrowser()
diff --git a/Engine/Plugins/Media/MediaFrameworkUtilities/Source/MediaFrameworkUtilitiesEditor/Private/AssetEditor/MediaBundleEditorToolkit.cpp b/Engine/Plugins/Media/MediaFrameworkUtilities/Source/MediaFrameworkUtilitiesEditor/Private/AssetEditor/MediaBundleEditorToolkit.cpp
--- a/Engine/Plugins/Media/MediaFrameworkUtilities/Source/MediaFrameworkUtilitiesEditor/Private/AssetEditor/MediaBundleEditorToolkit.cpp
+++ b/Engine/Plugins/Media/MediaFrameworkUtilities/Source/MediaFrameworkUtilitiesEditor/Private/AssetEditor/MediaBundleEditorToolkit.cpp
@@ -14,6 +14,7 @@
 #include "PropertyEditorModule.h"
 #include "UI/MediaFrameworkUtilitiesEditorStyle.h"
 #include "Widgets/Docking/SDockTab.h"
+#include "Subsystems/AssetEditorSubsystem.h"
  
 
 #define LOCTEXT_NAMESPACE "MediaBundleEditor"
@@ -215,7 +216,7 @@ void FMediaBundleEditorToolkit::ExtendToolBar()
 								IMaterialEditorModule* MaterialEditorModule = &FModuleManager::LoadModuleChecked<IMaterialEditorModule>("MaterialEditor");
 
 								UMaterialInterface* MaterialInterface = Asset->GetMaterial();
-								IAssetEditorInstance* EditorInstance = FAssetEditorManager::Get().FindEditorForAsset(MaterialInterface, true);
+								IAssetEditorInstance* EditorInstance = GEditor->GetEditorSubsystem<UAssetEditorSubsystem>()->FindEditorForAsset(MaterialInterface, true);
 								if (EditorInstance == nullptr)
 								{
 									UMaterialInstance* MaterialInstance = Cast<UMaterialInstance>(MaterialInterface);
diff --git a/Engine/Plugins/Media/MediaFrameworkUtilities/Source/MediaFrameworkUtilitiesEditor/Private/UI/MediaProfileMenuEntry.cpp b/Engine/Plugins/Media/MediaFrameworkUtilities/Source/MediaFrameworkUtilitiesEditor/Private/UI/MediaProfileMenuEntry.cpp
--- a/Engine/Plugins/Media/MediaFrameworkUtilities/Source/MediaFrameworkUtilitiesEditor/Private/UI/MediaProfileMenuEntry.cpp
+++ b/Engine/Plugins/Media/MediaFrameworkUtilities/Source/MediaFrameworkUtilitiesEditor/Private/UI/MediaProfileMenuEntry.cpp
@@ -17,8 +17,9 @@
 #include "Profile/MediaProfile.h"
 #include "Profile/MediaProfileSettings.h"
 #include "PropertyCustomizationHelpers.h"
-#include "Toolkits/AssetEditorManager.h"
+
 #include "UI/MediaFrameworkUtilitiesEditorStyle.h"
+#include "Subsystems/AssetEditorSubsystem.h"
 
 #define LOCTEXT_NAMESPACE "MediaProfileEditor"
 
@@ -34,7 +35,7 @@ struct FMediaProfileMenuEntryImpl
 			{
 				if (UMediaProfile* MediaProfile = IMediaProfileManager::Get().GetCurrentMediaProfile())
 				{
-					FAssetEditorManager::Get().OpenEditorForAsset(MediaProfile);
+					GEditor->GetEditorSubsystem<UAssetEditorSubsystem>()->OpenEditorForAsset(MediaProfile);
 				}
 			}),
 			FCanExecuteAction::CreateLambda([] { return IMediaProfileManager::Get().GetCurrentMediaProfile() != nullptr; }),
@@ -74,7 +75,7 @@ struct FMediaProfileMenuEntryImpl
 		{
 			GetMutableDefault<UMediaProfileEditorSettings>()->SetUserMediaProfile(NewAsset);
 			IMediaProfileManager::Get().SetCurrentMediaProfile(NewAsset);
-			FAssetEditorManager::Get().OpenEditorForAsset(NewAsset);
+			GEditor->GetEditorSubsystem<UAssetEditorSubsystem>()->OpenEditorForAsset(NewAsset);
 		}
 	}
 
diff --git a/Engine/Plugins/Media/MediaPlayerEditor/Source/MediaPlayerEditor/Private/Widgets/SMediaPlayerEditorMedia.cpp b/Engine/Plugins/Media/MediaPlayerEditor/Source/MediaPlayerEditor/Private/Widgets/SMediaPlayerEditorMedia.cpp
--- a/Engine/Plugins/Media/MediaPlayerEditor/Source/MediaPlayerEditor/Private/Widgets/SMediaPlayerEditorMedia.cpp
+++ b/Engine/Plugins/Media/MediaPlayerEditor/Source/MediaPlayerEditor/Private/Widgets/SMediaPlayerEditorMedia.cpp
@@ -8,7 +8,7 @@
 #include "EditorStyleSet.h"
 #include "Editor.h"
 #include "Containers/ArrayBuilder.h"
-#include "Toolkits/AssetEditorManager.h"
+
 #include "FileMediaSource.h"
 #include "MediaPlayer.h"
 #include "MediaPlaylist.h"
@@ -22,6 +22,7 @@
 #include "Widgets/Text/STextBlock.h"
 #include "ToolMenus.h"
 #include "Shared/MediaPlayerEditorMediaContext.h"
+#include "Subsystems/AssetEditorSubsystem.h"
 
 #define LOCTEXT_NAMESPACE "SMediaPlayerEditorMedia"
 
@@ -142,7 +143,7 @@ void SMediaPlayerEditorMedia::RegisterMenus()
 				{
 					if (Context->SelectedAsset)
 					{
-						FAssetEditorManager::Get().OpenEditorForAsset(Context->SelectedAsset);
+						GEditor->GetEditorSubsystem<UAssetEditorSubsystem>()->OpenEditorForAsset(Context->SelectedAsset);
 					}
 				}
 			})
diff --git a/Engine/Plugins/Media/TimecodeSynchronizer/Source/TimecodeSynchronizerEditor/Private/UI/TimecodeSynchronizerEditorLevelToolbar.cpp b/Engine/Plugins/Media/TimecodeSynchronizer/Source/TimecodeSynchronizerEditor/Private/UI/TimecodeSynchronizerEditorLevelToolbar.cpp
--- a/Engine/Plugins/Media/TimecodeSynchronizer/Source/TimecodeSynchronizerEditor/Private/UI/TimecodeSynchronizerEditorLevelToolbar.cpp
+++ b/Engine/Plugins/Media/TimecodeSynchronizer/Source/TimecodeSynchronizerEditor/Private/UI/TimecodeSynchronizerEditorLevelToolbar.cpp
@@ -15,8 +15,9 @@
 #include "Misc/FeedbackContext.h"
 #include "Modules/ModuleManager.h"
 #include "PropertyCustomizationHelpers.h"
-#include "Toolkits/AssetEditorManager.h"
+
 #include "Toolkits/AssetEditorToolkit.h"
+#include "Subsystems/AssetEditorSubsystem.h"
 
 
 #define LOCTEXT_NAMESPACE "TimecodeSynchronizerEditor"
@@ -168,7 +169,7 @@ void FTimecodeSynchronizerEditorLevelToolbar::AddObjectSubMenu(FMenuBuilder& Men
 
 void FTimecodeSynchronizerEditorLevelToolbar::OpenCurrentTimecodeSynchronizer()
 {
-	FAssetEditorManager::Get().OpenEditorForAsset(CurrentTimecodeSynchronizer.Get());
+	GEditor->GetEditorSubsystem<UAssetEditorSubsystem>()->OpenEditorForAsset(CurrentTimecodeSynchronizer.Get());
 }
 
 void FTimecodeSynchronizerEditorLevelToolbar::CreateNewTimecodeSynchronizer()
@@ -181,7 +182,7 @@ void FTimecodeSynchronizerEditorLevelToolbar::CreateNewTimecodeSynchronizer()
 		GetMutableDefault<UTimecodeSynchronizerEditorSettings>()->UserTimecodeSynchronizer = NewAsset;
 		GetMutableDefault<UTimecodeSynchronizerEditorSettings>()->SaveConfig();
 		CurrentTimecodeSynchronizer = NewAsset;
-		FAssetEditorManager::Get().OpenEditorForAsset(NewAsset);
+		GEditor->GetEditorSubsystem<UAssetEditorSubsystem>()->OpenEditorForAsset(NewAsset);
 	}
 }
 
diff --git a/Engine/Plugins/MovieScene/LevelSequenceEditor/Source/LevelSequenceEditor/Private/LevelSequenceEditorBlueprintLibrary.cpp b/Engine/Plugins/MovieScene/LevelSequenceEditor/Source/LevelSequenceEditor/Private/LevelSequenceEditorBlueprintLibrary.cpp
--- a/Engine/Plugins/MovieScene/LevelSequenceEditor/Source/LevelSequenceEditor/Private/LevelSequenceEditorBlueprintLibrary.cpp
+++ b/Engine/Plugins/MovieScene/LevelSequenceEditor/Source/LevelSequenceEditor/Private/LevelSequenceEditorBlueprintLibrary.cpp
@@ -7,6 +7,7 @@
 
 #include "Modules/ModuleManager.h"
 #include "LevelEditor.h"
+#include "Subsystems/AssetEditorSubsystem.h"
 
 namespace
 {
@@ -15,7 +16,7 @@ namespace
 
 bool ULevelSequenceEditorBlueprintLibrary::OpenLevelSequence(ULevelSequence* LevelSequence)
 {
-	return FAssetEditorManager::Get().OpenEditorForAsset(LevelSequence);
+	return GEditor->GetEditorSubsystem<UAssetEditorSubsystem>()->OpenEditorForAsset(LevelSequence);
 }
 
 ULevelSequence* ULevelSequenceEditorBlueprintLibrary::GetCurrentLevelSequence()
@@ -31,7 +32,7 @@ void ULevelSequenceEditorBlueprintLibrary::CloseLevelSequence()
 {
 	if (CurrentSequencer.IsValid())
 	{
-		FAssetEditorManager::Get().CloseAllEditorsForAsset(CurrentSequencer.Pin()->GetRootMovieSceneSequence());
+		GEditor->GetEditorSubsystem<UAssetEditorSubsystem>()->CloseAllEditorsForAsset(CurrentSequencer.Pin()->GetRootMovieSceneSequence());
 	}
 }
 
diff --git a/Engine/Plugins/MovieScene/LevelSequenceEditor/Source/LevelSequenceEditor/Private/LevelSequenceEditorModule.cpp b/Engine/Plugins/MovieScene/LevelSequenceEditor/Source/LevelSequenceEditor/Private/LevelSequenceEditorModule.cpp
--- a/Engine/Plugins/MovieScene/LevelSequenceEditor/Source/LevelSequenceEditor/Private/LevelSequenceEditorModule.cpp
+++ b/Engine/Plugins/MovieScene/LevelSequenceEditor/Source/LevelSequenceEditor/Private/LevelSequenceEditorModule.cpp
@@ -10,7 +10,7 @@
 #include "Framework/Commands/UICommandList.h"
 #include "Framework/MultiBox/MultiBoxExtender.h"
 #include "Framework/MultiBox/MultiBoxBuilder.h"
-#include "Toolkits/AssetEditorManager.h"
+
 #include "IAssetTypeActions.h"
 #include "AssetTools/LevelSequenceActions.h"
 #include "LevelSequenceEditorCommands.h"
@@ -37,6 +37,7 @@
 #include "SequencerSettings.h"
 #include "Misc/MovieSceneSequenceEditor_LevelSequence.h"
 #include "BlueprintAssetHandler.h"
+#include "Subsystems/AssetEditorSubsystem.h"
 
 #define LOCTEXT_NAMESPACE "LevelSequenceEditor"
 
@@ -354,7 +355,7 @@ protected:
 			GEditor->MoveViewportCamerasToActor(*NewActor, false);
 		}
 
-		FAssetEditorManager::Get().OpenEditorForAsset(NewAsset);
+		GEditor->GetEditorSubsystem<UAssetEditorSubsystem>()->OpenEditorForAsset(NewAsset);
 	}
 
 	/** Callback for creating a new level sequence asset in the level. */
diff --git a/Engine/Plugins/MovieScene/LevelSequenceEditor/Source/LevelSequenceEditor/Private/Misc/LevelSequenceEditorHelpers.cpp b/Engine/Plugins/MovieScene/LevelSequenceEditor/Source/LevelSequenceEditor/Private/Misc/LevelSequenceEditorHelpers.cpp
--- a/Engine/Plugins/MovieScene/LevelSequenceEditor/Source/LevelSequenceEditor/Private/Misc/LevelSequenceEditorHelpers.cpp
+++ b/Engine/Plugins/MovieScene/LevelSequenceEditor/Source/LevelSequenceEditor/Private/Misc/LevelSequenceEditorHelpers.cpp
@@ -24,7 +24,7 @@
 #include "EditorStyleSet.h"
 #include "LevelSequenceEditorModule.h"
 #include "Framework/Docking/TabManager.h"
-#include "Toolkits/AssetEditorManager.h"
+
 #include "Misc/LevelSequenceEditorSettings.h"
 #include "Widgets/Layout/SScrollBox.h"
 
@@ -33,6 +33,7 @@
 #include "MovieSceneToolsProjectSettings.h"
 #include "PropertyEditorModule.h"
 #include "Widgets/Docking/SDockTab.h"
+#include "Subsystems/AssetEditorSubsystem.h"
 
 /* LevelSequenceEditorHelpers
  *****************************************************************************/
@@ -201,7 +202,7 @@ private:
 		UObject* MasterSequenceAsset = LevelSequenceEditorHelpers::CreateLevelSequenceAsset(MasterSequenceAssetName, MasterSequencePackagePath);
 		if (MasterSequenceAsset)
 		{
-			FAssetEditorManager::Get().OpenEditorForAsset(MasterSequenceAsset);
+			GEditor->GetEditorSubsystem<UAssetEditorSubsystem>()->OpenEditorForAsset(MasterSequenceAsset);
 		
 			ILevelSequenceEditorModule& LevelSequenceEditorModule = FModuleManager::LoadModuleChecked<ILevelSequenceEditorModule>("LevelSequenceEditor");		
 			LevelSequenceEditorModule.OnMasterSequenceCreated().Broadcast(MasterSequenceAsset);
diff --git a/Engine/Plugins/MovieScene/MatineeToLevelSequence/Source/MatineeToLevelSequence/Private/MatineeToLevelSequenceModule.cpp b/Engine/Plugins/MovieScene/MatineeToLevelSequence/Source/MatineeToLevelSequence/Private/MatineeToLevelSequenceModule.cpp
--- a/Engine/Plugins/MovieScene/MatineeToLevelSequence/Source/MatineeToLevelSequence/Private/MatineeToLevelSequenceModule.cpp
+++ b/Engine/Plugins/MovieScene/MatineeToLevelSequence/Source/MatineeToLevelSequence/Private/MatineeToLevelSequenceModule.cpp
@@ -49,7 +49,7 @@
 #include "Matinee/InterpTrackMoveAxis.h"
 #include "MatineeUtils.h"
 #include "Editor.h"
-#include "Toolkits/AssetEditorManager.h"
+
 #include "LevelEditor.h"
 #include "MovieScene.h"
 #include "Evaluation/MovieSceneEvaluationTemplateInstance.h"
@@ -78,6 +78,7 @@
 #include "MovieSceneTimeHelpers.h"
 #include "Components/SkeletalMeshComponent.h"
 #include "Engine/StaticMeshActor.h"
+#include "Subsystems/AssetEditorSubsystem.h"
 
 #define LOCTEXT_NAMESPACE "MatineeToLevelSequence"
 
@@ -274,7 +275,7 @@ protected:
 			UObject* NewAsset = NewActors[0]->LevelSequence.TryLoad();
 			if (NewAsset)
 			{
-				FAssetEditorManager::Get().OpenEditorForAsset(NewAsset);
+				GEditor->GetEditorSubsystem<UAssetEditorSubsystem>()->OpenEditorForAsset(NewAsset);
 			}
 
 			FText NotificationText;
diff --git a/Engine/Plugins/Runtime/ApexDestruction/Source/ApexDestructionEditor/Private/AssetTypeActions_DestructibleMesh.cpp b/Engine/Plugins/Runtime/ApexDestruction/Source/ApexDestructionEditor/Private/AssetTypeActions_DestructibleMesh.cpp
--- a/Engine/Plugins/Runtime/ApexDestruction/Source/ApexDestructionEditor/Private/AssetTypeActions_DestructibleMesh.cpp
+++ b/Engine/Plugins/Runtime/ApexDestruction/Source/ApexDestructionEditor/Private/AssetTypeActions_DestructibleMesh.cpp
@@ -6,6 +6,7 @@
 #include "Engine/StaticMesh.h"
 #include "Widgets/Notifications/SNotificationList.h"
 #include "Framework/Notifications/NotificationManager.h"
+#include "Subsystems/AssetEditorSubsystem.h"
 
 #define LOCTEXT_NAMESPACE "AssetTypeActions"
 
@@ -43,7 +44,7 @@ void FAssetTypeActions_DestructibleMesh::ExecuteCreateDestructibleMeshes(TArray<
 				UDestructibleMesh* DestructibleMesh = DestructibleMeshEditorModule.CreateDestructibleMeshFromStaticMesh(StaticMesh->GetOuter(), StaticMesh, NAME_None, StaticMesh->GetFlags(), ErrorMsg);
 				if (DestructibleMesh)
 				{
-					FAssetEditorManager::Get().OpenEditorForAsset(DestructibleMesh);
+					GEditor->GetEditorSubsystem<UAssetEditorSubsystem>()->OpenEditorForAsset(DestructibleMesh);
 					NewAssets.Add(DestructibleMesh);
 				}
 				else if (!ErrorMsg.IsEmpty())
diff --git a/Engine/Plugins/Runtime/GameplayAbilities/Source/GameplayAbilitiesEditor/Private/SGameplayCueEditor.cpp b/Engine/Plugins/Runtime/GameplayAbilities/Source/GameplayAbilitiesEditor/Private/SGameplayCueEditor.cpp
--- a/Engine/Plugins/Runtime/GameplayAbilities/Source/GameplayAbilitiesEditor/Private/SGameplayCueEditor.cpp
+++ b/Engine/Plugins/Runtime/GameplayAbilities/Source/GameplayAbilitiesEditor/Private/SGameplayCueEditor.cpp
@@ -43,7 +43,7 @@
 #include "Widgets/Input/SHyperlink.h"
 #include "Widgets/Input/SSearchBox.h"
 #include "GameplayTagsModule.h"
-#include "Toolkits/AssetEditorManager.h"
+
 #include "AbilitySystemGlobals.h"
 #include "AssetToolsModule.h"
 #include "GameplayTagsEditorModule.h"
@@ -53,6 +53,7 @@
 #include "GameplayTagsEditor/Private/SAddNewGameplayTagWidget.h"
 
 #include "SGameplayCueEditor_Picker.h"
+#include "Subsystems/AssetEditorSubsystem.h"
 #define LOCTEXT_NAMESPACE "SGameplayCueEditor"
 
 static const FName CueTagColumnName("GameplayCueTags");
@@ -1456,7 +1457,7 @@ void SGameplayCueEditor::OpenEditorForNotify(FString NotifyFullPath)
 		{
 			FString AssetName = FPaths::GetBaseFilename(AssetRef.ToString());
 			UObject* AssetObject = FindObject<UObject>(pkg, *AssetName);
-			FAssetEditorManager::Get().OpenEditorForAsset(AssetObject);
+			GEditor->GetEditorSubsystem<UAssetEditorSubsystem>()->OpenEditorForAsset(AssetObject);
 		}
 	}
 }
diff --git a/Engine/Plugins/Tests/EditorTests/Source/EditorTests/Private/UnrealEd/BlueprintAutomationTests.cpp b/Engine/Plugins/Tests/EditorTests/Source/EditorTests/Private/UnrealEd/BlueprintAutomationTests.cpp
--- a/Engine/Plugins/Tests/EditorTests/Source/EditorTests/Private/UnrealEd/BlueprintAutomationTests.cpp
+++ b/Engine/Plugins/Tests/EditorTests/Source/EditorTests/Private/UnrealEd/BlueprintAutomationTests.cpp
@@ -34,7 +34,7 @@
 #include "ObjectTools.h"
 #include "ARFilter.h"
 #include "AssetRegistryModule.h"
-#include "Toolkits/AssetEditorManager.h"
+
 
 #include "Kismet2/KismetEditorUtilities.h"
 #include "EdGraphSchema_K2.h"
@@ -42,6 +42,7 @@
 #include "Kismet2/CompilerResultsLog.h"
 #include "DiffResults.h"
 #include "GraphDiffControl.h"
+#include "Subsystems/AssetEditorSubsystem.h"
 
 DEFINE_LOG_CATEGORY_STATIC(LogBlueprintAutomationTests, Log, All);
 
@@ -324,7 +325,7 @@ public:
 	 */
 	static void CloseBlueprint(UBlueprint* const BlueprintObj)
 	{
-		IAssetEditorInstance* EditorInst = FAssetEditorManager::Get().FindEditorForAsset(BlueprintObj, /*bool bFocusIfOpen =*/false);
+		IAssetEditorInstance* EditorInst = GEditor->GetEditorSubsystem<UAssetEditorSubsystem>()->FindEditorForAsset(BlueprintObj, /*bool bFocusIfOpen =*/false);
 		if (EditorInst != nullptr)
 		{
 			UE_LOG(LogBlueprintAutomationTests, Log, TEXT("Closing '%s' so we don't invalidate the open version when unloading it."), *BlueprintObj->GetName());
diff --git a/Engine/Plugins/Tests/EditorTests/Source/EditorTests/Private/UnrealEd/EditorAssetAutomationTests.cpp b/Engine/Plugins/Tests/EditorTests/Source/EditorTests/Private/UnrealEd/EditorAssetAutomationTests.cpp
--- a/Engine/Plugins/Tests/EditorTests/Source/EditorTests/Private/UnrealEd/EditorAssetAutomationTests.cpp
+++ b/Engine/Plugins/Tests/EditorTests/Source/EditorTests/Private/UnrealEd/EditorAssetAutomationTests.cpp
@@ -94,7 +94,7 @@
 #include "ObjectTools.h"
 #include "PackageTools.h"
 #include "AssetRegistryModule.h"
-#include "Toolkits/AssetEditorManager.h"
+
 #include "ContentBrowserModule.h"
 
 #include "Slate/SlateBrushAsset.h"
@@ -111,6 +111,7 @@
 #include "Factories/DataAssetFactory.h"
 #include "Factories/CurveFactory.h"
 #include "AssetExportTask.h"
+#include "Subsystems/AssetEditorSubsystem.h"
 
 #define LOCTEXT_NAMESPACE "EditorAssetAutomationTests"
 
@@ -909,7 +910,7 @@ namespace ImportExportAssetHelper
 				// Do not show progress window because we need to take screenshot
 				// ActiveTopLevelWindow does not get set to editor window for asset if progress modal window was open
 				const bool bShowProgressWindow = false;
-				if (FAssetEditorManager::Get().OpenEditorForAsset(ImportedAsset, EToolkitMode::Standalone, TSharedPtr<IToolkitHost>(), bShowProgressWindow))
+				if (GEditor->GetEditorSubsystem<UAssetEditorSubsystem>()->OpenEditorForAsset(ImportedAsset, EToolkitMode::Standalone, TSharedPtr<IToolkitHost>(), bShowProgressWindow))
 				{
 					State = EState::WaitForEditor;
 				}
@@ -934,7 +935,7 @@ namespace ImportExportAssetHelper
 				if (!ActiveWindowTitle.StartsWith(ImportedAsset->GetName()))
 				{
 					//Bring the asset editor to the front
-					FAssetEditorManager::Get().FindEditorForAsset(ImportedAsset, true);
+					GEditor->GetEditorSubsystem<UAssetEditorSubsystem>()->FindEditorForAsset(ImportedAsset, true);
 				}
 
 				State = EState::Screenshot;
@@ -977,7 +978,7 @@ namespace ImportExportAssetHelper
 				}
 
 				//Close the editor
-				FAssetEditorManager::Get().CloseAllAssetEditors();
+				GEditor->GetEditorSubsystem<UAssetEditorSubsystem>()->CloseAllAssetEditors();
 
 				State = EState::Export;
 			}
diff --git a/Engine/Plugins/Tests/EditorTests/Source/EditorTests/Private/UnrealEd/EditorBuildPromotionTests.cpp b/Engine/Plugins/Tests/EditorTests/Source/EditorTests/Private/UnrealEd/EditorBuildPromotionTests.cpp
--- a/Engine/Plugins/Tests/EditorTests/Source/EditorTests/Private/UnrealEd/EditorBuildPromotionTests.cpp
+++ b/Engine/Plugins/Tests/EditorTests/Source/EditorTests/Private/UnrealEd/EditorBuildPromotionTests.cpp
@@ -67,7 +67,7 @@
 #include "PackageHelperFunctions.h"
 
 //Materials
-#include "Toolkits/AssetEditorManager.h"
+
 #include "Editor/MaterialEditor/Private/MaterialEditor.h"
 #include "Materials/MaterialExpressionConstant3Vector.h"
 
@@ -101,6 +101,7 @@
 
 #include "LevelEditor.h"
 #include "ObjectTools.h"
+#include "Subsystems/AssetEditorSubsystem.h"
 
 
 #define LOCTEXT_NAMESPACE "EditorBuildPromotionTests"
@@ -756,7 +757,7 @@ namespace BuildPromotionTestHelper
 		void OpenAssetEditor()
 		{
 			UObject* CurrentAsset = Assets[AssetIndex].Asset;
-			if (FAssetEditorManager::Get().OpenEditorForAsset(CurrentAsset))
+			if (GEditor->GetEditorSubsystem<UAssetEditorSubsystem>()->OpenEditorForAsset(CurrentAsset))
 			{
 				CurrentStage = EState::WaitForEditor;
 			}
@@ -784,7 +785,7 @@ namespace BuildPromotionTestHelper
 				if (!ActiveWindowTitle.StartsWith(CurrentAsset->GetName()))
 				{
 					//Bring the asset editor to the front
-					FAssetEditorManager::Get().FindEditorForAsset(CurrentAsset, true);
+					GEditor->GetEditorSubsystem<UAssetEditorSubsystem>()->FindEditorForAsset(CurrentAsset, true);
 				}
 
 				UE_LOG(LogEditorBuildPromotionTests, Display, TEXT("Opened asset (%s)"), *CurrentAsset->GetClass()->GetName());
@@ -874,7 +875,7 @@ namespace BuildPromotionTestHelper
 			}
 
 			//close editor
-			FAssetEditorManager::Get().CloseAllAssetEditors();
+			GEditor->GetEditorSubsystem<UAssetEditorSubsystem>()->CloseAllAssetEditors();
 
 			//Add to level
 			FVector PlaceLocation(-1090,970,160);
@@ -1381,7 +1382,7 @@ namespace BuildPromotionTestHelper
 				if (SCTestMat)
 				{
 					//Open the asset editor
-					FAssetEditorManager::Get().OpenEditorForAsset(SCTestMat);
+					GEditor->GetEditorSubsystem<UAssetEditorSubsystem>()->OpenEditorForAsset(SCTestMat);
 					UE_LOG(LogEditorBuildPromotionTests, Display, TEXT("Opened the material editor for: %s"), *SCTestMat->GetName());
 
 					FString PackageFileName = FPackageName::LongPackageNameToFilename(MaterialData.PackageName.ToString(), FPackageName::GetAssetPackageExtension());
@@ -1463,7 +1464,7 @@ namespace BuildPromotionTestHelper
 				UAutomationTestSettings const* AutomationTestSettings = GetDefault<UAutomationTestSettings>();
 				check(AutomationTestSettings);
 
-				IAssetEditorInstance* AssetEditor = FAssetEditorManager::Get().FindEditorForAsset(SCTestMat, true);
+				IAssetEditorInstance* AssetEditor = GEditor->GetEditorSubsystem<UAssetEditorSubsystem>()->FindEditorForAsset(SCTestMat, true);
 				FMaterialEditor* MaterialEditor = (FMaterialEditor*)AssetEditor;
 
 				//Pick a random color
@@ -1566,7 +1567,7 @@ namespace BuildPromotionTestHelper
 		bool ContentBrowser_SourceControl_Part5()
 		{
 			UE_LOG(LogEditorBuildPromotionTests, Display, TEXT("Closed the material editor"));
-			FAssetEditorManager::Get().CloseAllAssetEditors();
+			GEditor->GetEditorSubsystem<UAssetEditorSubsystem>()->CloseAllAssetEditors();
 			return true;
 		}
 
@@ -1775,7 +1776,7 @@ namespace BuildPromotionTestHelper
 		bool ContentBrowser_AssignAMaterial()
 		{
 			// SETUP
-			FAssetEditorManager::Get().CloseAllAssetEditors();
+			GEditor->GetEditorSubsystem<UAssetEditorSubsystem>()->CloseAllAssetEditors();
 			UAutomationTestSettings const* AutomationTestSettings = GetDefault<UAutomationTestSettings>();
 			check(AutomationTestSettings);
 
@@ -2064,7 +2065,7 @@ namespace BuildPromotionTestHelper
 			{
 				//Open the level script blueprint
 				ULevelScriptBlueprint* LSB = PlacedBlueprint->GetLevel()->GetLevelScriptBlueprint(false);
-				FAssetEditorManager::Get().OpenEditorForAsset(LSB);
+				GEditor->GetEditorSubsystem<UAssetEditorSubsystem>()->OpenEditorForAsset(LSB);
 				Test->AddInfo(TEXT("Opened the level script blueprint"));
 			}
 			return true;
@@ -2081,7 +2082,7 @@ namespace BuildPromotionTestHelper
 		{
 			if (BlueprintObject)
 			{
-				IAssetEditorInstance* AssetEditor = FAssetEditorManager::Get().FindEditorForAsset(BlueprintObject, true);
+				IAssetEditorInstance* AssetEditor = GEditor->GetEditorSubsystem<UAssetEditorSubsystem>()->FindEditorForAsset(BlueprintObject, true);
 				FBlueprintEditor* BlueprintEditor = (FBlueprintEditor*)AssetEditor;
 
 				UEdGraph* EventGraph = FBlueprintEditorUtils::FindEventGraph(BlueprintObject);
@@ -2182,7 +2183,7 @@ namespace BuildPromotionTestHelper
 			if (BlueprintObject)
 			{
 				Test->AddInfo(TEXT("Closing the blueprint editor"));
-				FAssetEditorManager::Get().CloseAllAssetEditors();
+				GEditor->GetEditorSubsystem<UAssetEditorSubsystem>()->CloseAllAssetEditors();
 				//UE_LOG(LogEditorBuildPromotionTests, Display, TEXT("Saving the blueprint"));
 				//SaveBlueprint();
 			}
diff --git a/Engine/Plugins/Tests/EditorTests/Source/EditorTests/Private/UnrealEd/EditorOpenAssetAutomationTests.cpp b/Engine/Plugins/Tests/EditorTests/Source/EditorTests/Private/UnrealEd/EditorOpenAssetAutomationTests.cpp
--- a/Engine/Plugins/Tests/EditorTests/Source/EditorTests/Private/UnrealEd/EditorOpenAssetAutomationTests.cpp
+++ b/Engine/Plugins/Tests/EditorTests/Source/EditorTests/Private/UnrealEd/EditorOpenAssetAutomationTests.cpp
@@ -9,7 +9,9 @@
 #include "Tests/AutomationTestSettings.h"
 #include "Tests/AutomationEditorCommon.h"
 #include "Tests/AutomationCommon.h"
-#include "Toolkits/AssetEditorManager.h"
+
+#include "Subsystems/AssetEditorSubsystem.h"
+#include "Editor.h"
 
 /**
  * Test to open the sub editor windows for a specified list of assets.
@@ -30,7 +32,7 @@ void FOpenAssetEditors::GetTests(TArray<FString>& OutBeautifiedNames, TArray<FSt
 bool FOpenAssetEditors::RunTest(const FString& LongAssetPath)
 {
 	//start with all editors closed
-	FAssetEditorManager::Get().CloseAllAssetEditors();
+	GEditor->GetEditorSubsystem<UAssetEditorSubsystem>()->CloseAllAssetEditors();
 
 	// below is all latent action, so before sending there, verify the asset exists
 	UObject* Object = FSoftObjectPath(LongAssetPath).TryLoad();
diff --git a/Engine/Plugins/VirtualProduction/Takes/Source/TakeRecorder/Private/Recorder/TakeRecorder.cpp b/Engine/Plugins/VirtualProduction/Takes/Source/TakeRecorder/Private/Recorder/TakeRecorder.cpp
--- a/Engine/Plugins/VirtualProduction/Takes/Source/TakeRecorder/Private/Recorder/TakeRecorder.cpp
+++ b/Engine/Plugins/VirtualProduction/Takes/Source/TakeRecorder/Private/Recorder/TakeRecorder.cpp
@@ -24,7 +24,7 @@
 
 // UnrealEd includes
 #include "Editor.h"
-#include "Toolkits/AssetEditorManager.h"
+
 #include "ObjectTools.h"
 #include "UObject/GCObjectScopeGuard.h"
 
@@ -42,6 +42,7 @@
 // LevelEditor includes
 #include "ILevelViewport.h"
 #include "LevelEditor.h"
+#include "Subsystems/AssetEditorSubsystem.h"
 
 #define LOCTEXT_NAMESPACE "TakeRecorder"
 
@@ -522,9 +523,9 @@ bool UTakeRecorder::CreateDestinationAsset(const TCHAR* AssetPathFormat, ULevelS
 bool UTakeRecorder::InitializeSequencer(FText* OutError)
 {
 	// Open the sequence and set the sequencer ptr
-	FAssetEditorManager::Get().OpenEditorForAsset(SequenceAsset);
+	GEditor->GetEditorSubsystem<UAssetEditorSubsystem>()->OpenEditorForAsset(SequenceAsset);
 
-	IAssetEditorInstance*        AssetEditor         = FAssetEditorManager::Get().FindEditorForAsset(SequenceAsset, false);
+	IAssetEditorInstance*        AssetEditor         = GEditor->GetEditorSubsystem<UAssetEditorSubsystem>()->FindEditorForAsset(SequenceAsset, false);
 	ILevelSequenceEditorToolkit* LevelSequenceEditor = static_cast<ILevelSequenceEditorToolkit*>(AssetEditor);
 
 	WeakSequencer = LevelSequenceEditor ? LevelSequenceEditor->GetSequencer() : nullptr;
diff --git a/Engine/Plugins/VirtualProduction/Takes/Source/TakeRecorder/Private/ScopedSequencerPanel.cpp b/Engine/Plugins/VirtualProduction/Takes/Source/TakeRecorder/Private/ScopedSequencerPanel.cpp
--- a/Engine/Plugins/VirtualProduction/Takes/Source/TakeRecorder/Private/ScopedSequencerPanel.cpp
+++ b/Engine/Plugins/VirtualProduction/Takes/Source/TakeRecorder/Private/ScopedSequencerPanel.cpp
@@ -21,11 +21,12 @@
 #include "EditorStyleSet.h"
 
 // UnrealEd includes
-#include "Toolkits/AssetEditorManager.h"
+#include "Subsystems/AssetEditorSubsystem.h"
 
 // LevelSequenceEditor includes
 #include "ILevelSequenceEditorToolkit.h"
 
+
 FScopedSequencerPanel::FScopedSequencerPanel(const TAttribute<ULevelSequence*>& InLevelSequenceAttribute)
 	: LevelSequenceAttribute(InLevelSequenceAttribute)
 {
@@ -42,14 +43,14 @@ FScopedSequencerPanel::~FScopedSequencerPanel()
 
 bool FScopedSequencerPanel::IsOpen(ULevelSequence* InSequence)
 {
-	return FAssetEditorManager::Get().FindEditorForAsset(InSequence, false) != nullptr;
+	return GEditor->GetEditorSubsystem<UAssetEditorSubsystem>()->FindEditorForAsset(InSequence, false) != nullptr;
 }
 
 void FScopedSequencerPanel::Open(ULevelSequence* InSequence)
 {
-	FAssetEditorManager::Get().OpenEditorForAsset(InSequence);
+	GEditor->GetEditorSubsystem<UAssetEditorSubsystem>()->OpenEditorForAsset(InSequence);
 
-	IAssetEditorInstance*        AssetEditor = FAssetEditorManager::Get().FindEditorForAsset(InSequence, false);
+	IAssetEditorInstance*        AssetEditor = GEditor->GetEditorSubsystem<UAssetEditorSubsystem>()->FindEditorForAsset(InSequence, false);
 	ILevelSequenceEditorToolkit* LevelSequenceEditor = static_cast<ILevelSequenceEditorToolkit*>(AssetEditor);
 
 	if (LevelSequenceEditor->GetSequencer())
@@ -62,7 +63,7 @@ void FScopedSequencerPanel::Open(ULevelSequence* InSequence)
 
 void FScopedSequencerPanel::Close(ULevelSequence* InSequence)
 {
-	FAssetEditorManager::Get().CloseAllEditorsForAsset(InSequence);
+	GEditor->GetEditorSubsystem<UAssetEditorSubsystem>()->CloseAllEditorsForAsset(InSequence);
 }
 
 ECheckBoxState FScopedSequencerPanel::GetToggleCheckState() const
diff --git a/Engine/Plugins/VirtualProduction/Takes/Source/TakeRecorderSources/Private/TakeRecorderWorldSource.cpp b/Engine/Plugins/VirtualProduction/Takes/Source/TakeRecorderSources/Private/TakeRecorderWorldSource.cpp
--- a/Engine/Plugins/VirtualProduction/Takes/Source/TakeRecorderSources/Private/TakeRecorderWorldSource.cpp
+++ b/Engine/Plugins/VirtualProduction/Takes/Source/TakeRecorderSources/Private/TakeRecorderWorldSource.cpp
@@ -13,8 +13,9 @@
 #include "ISequencer.h"
 #include "MovieSceneFolder.h"
 
-#include "Toolkits/AssetEditorManager.h"
+
 #include "ILevelSequenceEditorToolkit.h"
+#include "Subsystems/AssetEditorSubsystem.h"
 
 UTakeRecorderWorldSourceSettings::UTakeRecorderWorldSourceSettings(const FObjectInitializer& ObjInit)
 	: Super(ObjInit)
@@ -126,7 +127,7 @@ void UTakeRecorderWorldSource::AutotrackActors(class ULevelSequence* InSequence,
 		return;
 	}
 
-	IAssetEditorInstance*        AssetEditor = FAssetEditorManager::Get().FindEditorForAsset(InSequence, false);
+	IAssetEditorInstance*        AssetEditor = GEditor->GetEditorSubsystem<UAssetEditorSubsystem>()->FindEditorForAsset(InSequence, false);
 	ILevelSequenceEditorToolkit* LevelSequenceEditor = static_cast<ILevelSequenceEditorToolkit*>(AssetEditor);
 
 	TSharedPtr<ISequencer> SequencerPtr = LevelSequenceEditor ? LevelSequenceEditor->GetSequencer() : nullptr;
diff --git a/Engine/Source/Editor/AnimGraph/Private/AnimBlueprintPinInfoDetails.cpp b/Engine/Source/Editor/AnimGraph/Private/AnimBlueprintPinInfoDetails.cpp
--- a/Engine/Source/Editor/AnimGraph/Private/AnimBlueprintPinInfoDetails.cpp
+++ b/Engine/Source/Editor/AnimGraph/Private/AnimBlueprintPinInfoDetails.cpp
@@ -10,10 +10,11 @@
 #include "DetailWidgetRow.h"
 #include "Widgets/Input/SEditableTextBox.h"
 #include "Widgets/Layout/SBox.h"
-#include "Toolkits/AssetEditorManager.h"
+
 #include "Animation/AnimBlueprint.h"
 #include "IAnimationBlueprintEditor.h"
 #include "AnimationGraphSchema.h"
+#include "Subsystems/AssetEditorSubsystem.h"
 
 #define LOCTEXT_NAMESPACE "AnimBlueprintPinInfoDetails"
 
@@ -109,7 +110,7 @@ void FAnimBlueprintFunctionPinInfoDetails::HandlePinTypeChanged(const FEdGraphPi
 		TypePropertyHandle->NotifyPostChange(EPropertyChangeType::ValueSet);
 
 		UAnimBlueprint* AnimBlueprint = WeakOuterNode->GetAnimBlueprint();
-		IAssetEditorInstance* AssetEditor = FAssetEditorManager::Get().FindEditorForAsset(AnimBlueprint, false);
+		IAssetEditorInstance* AssetEditor = GEditor->GetEditorSubsystem<UAssetEditorSubsystem>()->FindEditorForAsset(AnimBlueprint, false);
 		check(AssetEditor->GetEditorName() == "AnimationBlueprintEditor");
 
 		static_cast<IAnimationBlueprintEditor*>(AssetEditor)->SetLastGraphPinTypeUsed(InPinType);
diff --git a/Engine/Source/Editor/AnimGraph/Private/AnimGraphNode_Base.cpp b/Engine/Source/Editor/AnimGraph/Private/AnimGraphNode_Base.cpp
--- a/Engine/Source/Editor/AnimGraph/Private/AnimGraphNode_Base.cpp
+++ b/Engine/Source/Editor/AnimGraph/Private/AnimGraphNode_Base.cpp
@@ -12,10 +12,11 @@
 #include "AnimNodeEditModes.h"
 #include "AnimationGraph.h"
 #include "EditorModeManager.h"
-#include "Toolkits/AssetEditorManager.h"
+
 #include "AnimationEditorUtils.h"
 #include "UObject/UnrealType.h"
 #include "Kismet2/CompilerResultsLog.h"
+#include "Subsystems/AssetEditorSubsystem.h"
 
 #define LOCTEXT_NAMESPACE "UAnimGraphNode_Base"
 
@@ -128,7 +129,7 @@ void UAnimGraphNode_Base::JumpToDefinition() const
 {
 	if (UObject* HyperlinkTarget = GetJumpTargetForDoubleClick())
 	{
-		FAssetEditorManager::Get().OpenEditorForAsset(HyperlinkTarget);
+		GEditor->GetEditorSubsystem<UAssetEditorSubsystem>()->OpenEditorForAsset(HyperlinkTarget);
 	}
 }
 
diff --git a/Engine/Source/Editor/AnimGraph/Private/AnimGraphNode_SubInput.cpp b/Engine/Source/Editor/AnimGraph/Private/AnimGraphNode_SubInput.cpp
--- a/Engine/Source/Editor/AnimGraph/Private/AnimGraphNode_SubInput.cpp
+++ b/Engine/Source/Editor/AnimGraph/Private/AnimGraphNode_SubInput.cpp
@@ -12,7 +12,7 @@
 #include "DetailLayoutBuilder.h"
 #include "Widgets/Images/SImage.h"
 #include "Widgets/Text/STextBlock.h"
-#include "Toolkits/AssetEditorManager.h"
+
 #include "Animation/AnimBlueprint.h"
 #include "IAnimationBlueprintEditor.h"
 #include "AnimationGraphSchema.h"
@@ -22,6 +22,7 @@
 #include "K2Node_CallFunction.h"
 #include "Containers/Ticker.h"
 #include "Widgets/Input/SCheckBox.h"
+#include "Subsystems/AssetEditorSubsystem.h"
 
 #define LOCTEXT_NAMESPACE "SubInputNode"
 
@@ -146,7 +147,7 @@ void UAnimGraphNode_SubInput::HandleInputPinArrayChanged()
 
 			if(Input.Type.PinCategory == NAME_None)
 			{
-				IAssetEditorInstance* AssetEditor = FAssetEditorManager::Get().FindEditorForAsset(AnimBlueprint, false);
+				IAssetEditorInstance* AssetEditor = GEditor->GetEditorSubsystem<UAssetEditorSubsystem>()->FindEditorForAsset(AnimBlueprint, false);
 				check(AssetEditor->GetEditorName() == "AnimationBlueprintEditor");
 				IAnimationBlueprintEditor* AnimationBlueprintEditor = static_cast<IAnimationBlueprintEditor*>(AssetEditor);
 				Input.Type = AnimationBlueprintEditor->GetLastGraphPinTypeUsed();
@@ -271,7 +272,7 @@ void UAnimGraphNode_SubInput::PostPlacedNewNode()
 			if(UAnimGraphNode_SubInput* SubInputNode = WeakThis.Get())
 			{
 				// refresh the BP editor's details panel in case we are viewing the graph
-				IAssetEditorInstance* AssetEditor = FAssetEditorManager::Get().FindEditorForAsset(SubInputNode->GetAnimBlueprint(), false);
+				IAssetEditorInstance* AssetEditor = GEditor->GetEditorSubsystem<UAssetEditorSubsystem>()->FindEditorForAsset(SubInputNode->GetAnimBlueprint(), false);
 				check(AssetEditor->GetEditorName() == "AnimationBlueprintEditor");
 				IAnimationBlueprintEditor* AnimationBlueprintEditor = static_cast<IAnimationBlueprintEditor*>(AssetEditor);
 				AnimationBlueprintEditor->RefreshInspector();
diff --git a/Engine/Source/Editor/AnimationBlueprintEditor/Private/AnimationBlueprintEditor.cpp b/Engine/Source/Editor/AnimationBlueprintEditor/Private/AnimationBlueprintEditor.cpp
--- a/Engine/Source/Editor/AnimationBlueprintEditor/Private/AnimationBlueprintEditor.cpp
+++ b/Engine/Source/Editor/AnimationBlueprintEditor/Private/AnimationBlueprintEditor.cpp
@@ -74,6 +74,7 @@
 #include "EdGraphNode_Comment.h"
 #include "AnimStateNodeBase.h"
 #include "AnimStateEntryNode.h"
+#include "Subsystems/AssetEditorSubsystem.h"
 
 #define LOCTEXT_NAMESPACE "AnimationBlueprintEditor"
 
@@ -1101,7 +1102,7 @@ void FAnimationBlueprintEditor::OnOpenRelatedAsset()
 				UAnimationAsset* AnimAsset = Node->GetAnimationAsset();
 				if(AnimAsset)
 				{
-					FAssetEditorManager::Get().OpenEditorForAsset(AnimAsset, Mode);
+					GEditor->GetEditorSubsystem<UAssetEditorSubsystem>()->OpenEditorForAsset(AnimAsset, Mode);
 				}
 			}
 		}
@@ -1135,7 +1136,7 @@ void FAnimationBlueprintEditor::Compile()
 		if(CurrentDebugObject)
 		{
 			// Force close any asset editors that are using the AnimScriptInstance (such as the Property Matrix), the class will be garbage collected
-			FAssetEditorManager::Get().CloseOtherEditors(CurrentDebugObject, nullptr);
+			GEditor->GetEditorSubsystem<UAssetEditorSubsystem>()->CloseOtherEditors(CurrentDebugObject, nullptr);
 			DebuggedMeshComponent = CurrentDebugObject->GetSkelMeshComponent();
 		}
 	}
@@ -1444,7 +1445,7 @@ UObject* FAnimationBlueprintEditor::HandleGetObject()
 
 void FAnimationBlueprintEditor::HandleOpenNewAsset(UObject* InNewAsset)
 {
-	FAssetEditorManager::Get().OpenEditorForAsset(InNewAsset);
+	GEditor->GetEditorSubsystem<UAssetEditorSubsystem>()->OpenEditorForAsset(InNewAsset);
 }
 
 void FAnimationBlueprintEditor::AddReferencedObjects( FReferenceCollector& Collector )
diff --git a/Engine/Source/Editor/AnimationBlueprintEditor/Public/IAnimationBlueprintEditorModule.h b/Engine/Source/Editor/AnimationBlueprintEditor/Public/IAnimationBlueprintEditorModule.h
--- a/Engine/Source/Editor/AnimationBlueprintEditor/Public/IAnimationBlueprintEditorModule.h
+++ b/Engine/Source/Editor/AnimationBlueprintEditor/Public/IAnimationBlueprintEditorModule.h
@@ -20,7 +20,7 @@ public:
 	 *
 	 * Note: This function should not be called directly, use one of the following instead:
 	 *	- FKismetEditorUtilities::BringKismetToFocusAttentionOnObject
-	 *  - FAssetEditorManager::Get().OpenEditorForAsset
+	 *  - GEditor->GetEditorSubsystem<UAssetEditorSubsystem>()->OpenEditorForAsset
 	 *
 	 * @param	Mode					Mode that this editor should operate in
 	 * @param	InitToolkitHost			When Mode is WorldCentric, this is the level editor instance to spawn this editor within
diff --git a/Engine/Source/Editor/AnimationModifiers/Private/SAnimationModifiersTab.cpp b/Engine/Source/Editor/AnimationModifiers/Private/SAnimationModifiersTab.cpp
--- a/Engine/Source/Editor/AnimationModifiers/Private/SAnimationModifiersTab.cpp
+++ b/Engine/Source/Editor/AnimationModifiers/Private/SAnimationModifiersTab.cpp
@@ -18,7 +18,7 @@
 #include "Dialogs/Dialogs.h"
 #include "Editor.h"
 #include "ScopedTransaction.h"
-#include "Toolkits/AssetEditorManager.h"
+
 #include "Editor.h"
 #include "Widgets/Input/SComboButton.h"
 #include "Widgets/Input/SButton.h"
@@ -26,6 +26,7 @@
 #include "Widgets/Input/SMenuAnchor.h"
 #include "Engine/BlueprintGeneratedClass.h"
 #include "AnimationModifierHelpers.h"
+#include "Subsystems/AssetEditorSubsystem.h"
 
 #define LOCTEXT_NAMESPACE "SAnimationModifiersTab"
 
@@ -41,7 +42,7 @@ SAnimationModifiersTab::~SAnimationModifiersTab()
 		GEditor->UnregisterForUndo(this);
 	}
 
-	FAssetEditorManager::Get().OnAssetOpenedInEditor().RemoveAll(this);	
+	GEditor->GetEditorSubsystem<UAssetEditorSubsystem>()->OnAssetOpenedInEditor().RemoveAll(this);	
 }
 
 void SAnimationModifiersTab::Tick(const FGeometry& AllottedGeometry, const double InCurrentTime, const float InDeltaTime)
@@ -172,7 +173,7 @@ void SAnimationModifiersTab::Construct(const FArguments& InArgs)
 	{
 		GEditor->RegisterForUndo(this);
 	}
-	FAssetEditorManager::Get().OnAssetOpenedInEditor().AddSP(this, &SAnimationModifiersTab::OnAssetOpened);
+	GEditor->GetEditorSubsystem<UAssetEditorSubsystem>()->OnAssetOpenedInEditor().AddSP(this, &SAnimationModifiersTab::OnAssetOpened);
 }
 
 void SAnimationModifiersTab::OnModifierPicked(UClass* PickedClass)
@@ -360,7 +361,7 @@ void SAnimationModifiersTab::OnOpenModifier(const TWeakObjectPtr<UAnimationModif
 		UBlueprint* Blueprint = Cast<UBlueprint>(BPGeneratedClass->ClassGeneratedBy);
 		if (Blueprint)
 		{
-			FAssetEditorManager::Get().OpenEditorForAsset(Blueprint);
+			GEditor->GetEditorSubsystem<UAssetEditorSubsystem>()->OpenEditorForAsset(Blueprint);
 		}
 	}
 }
diff --git a/Engine/Source/Editor/BehaviorTreeEditor/Private/AssetTypeActions_BehaviorTree.cpp b/Engine/Source/Editor/BehaviorTreeEditor/Private/AssetTypeActions_BehaviorTree.cpp
--- a/Engine/Source/Editor/BehaviorTreeEditor/Private/AssetTypeActions_BehaviorTree.cpp
+++ b/Engine/Source/Editor/BehaviorTreeEditor/Private/AssetTypeActions_BehaviorTree.cpp
@@ -12,6 +12,7 @@
 #include "SBehaviorTreeDiff.h"
 
 #include "AIModule.h"
+#include "Subsystems/AssetEditorSubsystem.h"
 
 #define LOCTEXT_NAMESPACE "AssetTypeActions"
 
@@ -35,7 +36,7 @@ void FAssetTypeActions_BehaviorTree::OpenAssetEditor( const TArray<UObject*>& In
 			if(BehaviorTree->BlackboardAsset != nullptr)
 			{
 				const bool bFocusIfOpen = false;
-				FBehaviorTreeEditor* ExistingInstance = static_cast<FBehaviorTreeEditor*>(FAssetEditorManager::Get().FindEditorForAsset(BehaviorTree->BlackboardAsset, bFocusIfOpen));
+				FBehaviorTreeEditor* ExistingInstance = static_cast<FBehaviorTreeEditor*>(GEditor->GetEditorSubsystem<UAssetEditorSubsystem>()->FindEditorForAsset(BehaviorTree->BlackboardAsset, bFocusIfOpen));
 				if(ExistingInstance != nullptr && ExistingInstance->GetBehaviorTree() == nullptr)
 				{
 					ExistingInstance->InitBehaviorTreeEditor(Mode, EditWithinLevelEditor, BehaviorTree);
diff --git a/Engine/Source/Editor/BehaviorTreeEditor/Private/BehaviorTreeEditor.cpp b/Engine/Source/Editor/BehaviorTreeEditor/Private/BehaviorTreeEditor.cpp
--- a/Engine/Source/Editor/BehaviorTreeEditor/Private/BehaviorTreeEditor.cpp
+++ b/Engine/Source/Editor/BehaviorTreeEditor/Private/BehaviorTreeEditor.cpp
@@ -60,6 +60,7 @@
 #include "IContentBrowserSingleton.h"
 #include "ContentBrowserModule.h"
 #include "Widgets/Docking/SDockTab.h"
+#include "Subsystems/AssetEditorSubsystem.h"
 
 #define LOCTEXT_NAMESPACE "BehaviorTreeEditor"
 
@@ -1156,9 +1157,9 @@ void FBehaviorTreeEditor::OnNodeDoubleClicked(class UEdGraphNode* Node)
 		UBTTask_RunBehavior* SubtreeTask = MyNode->ParentNode ? Cast<UBTTask_RunBehavior>(MyNode->ParentNode->NodeInstance) : NULL;
 		if (SubtreeTask && SubtreeTask->GetSubtreeAsset())
 		{
-			FAssetEditorManager::Get().OpenEditorForAsset(SubtreeTask->GetSubtreeAsset());
+			GEditor->GetEditorSubsystem<UAssetEditorSubsystem>()->OpenEditorForAsset(SubtreeTask->GetSubtreeAsset());
 
-			IBehaviorTreeEditor* ChildNodeEditor = static_cast<IBehaviorTreeEditor*>(FAssetEditorManager::Get().FindEditorForAsset(SubtreeTask->GetSubtreeAsset(), true));
+			IBehaviorTreeEditor* ChildNodeEditor = static_cast<IBehaviorTreeEditor*>(GEditor->GetEditorSubsystem<UAssetEditorSubsystem>()->FindEditorForAsset(SubtreeTask->GetSubtreeAsset(), true));
 			if (ChildNodeEditor)
 			{
 				ChildNodeEditor->InitializeDebuggerState(Debugger.Get());
@@ -1211,9 +1212,9 @@ void FBehaviorTreeEditor::OnNodeDoubleClicked(class UEdGraphNode* Node)
 		{
 			if (RunTask->GetSubtreeAsset())
 			{
-				FAssetEditorManager::Get().OpenEditorForAsset(RunTask->GetSubtreeAsset());
+				GEditor->GetEditorSubsystem<UAssetEditorSubsystem>()->OpenEditorForAsset(RunTask->GetSubtreeAsset());
 
-				IBehaviorTreeEditor* ChildNodeEditor = static_cast<IBehaviorTreeEditor*>(FAssetEditorManager::Get().FindEditorForAsset(RunTask->GetSubtreeAsset(), true));
+				IBehaviorTreeEditor* ChildNodeEditor = static_cast<IBehaviorTreeEditor*>(GEditor->GetEditorSubsystem<UAssetEditorSubsystem>()->FindEditorForAsset(RunTask->GetSubtreeAsset(), true));
 				if (ChildNodeEditor)
 				{
 					ChildNodeEditor->InitializeDebuggerState(Debugger.Get());
@@ -1231,7 +1232,7 @@ void FBehaviorTreeEditor::OnNodeDoubleClicked(class UEdGraphNode* Node)
 		UBlueprint* BlueprintOb = FindObject<UBlueprint>(Pkg, *ClassName);
 		if(BlueprintOb)
 		{
-			FAssetEditorManager::Get().OpenEditorForAsset(BlueprintOb);
+			GEditor->GetEditorSubsystem<UAssetEditorSubsystem>()->OpenEditorForAsset(BlueprintOb);
 		}
 	}
 }
@@ -1616,9 +1617,9 @@ void FBehaviorTreeEditor::DebuggerSwitchAsset(UBehaviorTree* NewAsset)
 {
 	if (NewAsset)
 	{
-		FAssetEditorManager::Get().OpenEditorForAsset(NewAsset);
+		GEditor->GetEditorSubsystem<UAssetEditorSubsystem>()->OpenEditorForAsset(NewAsset);
 
-		IBehaviorTreeEditor* ChildNodeEditor = static_cast<IBehaviorTreeEditor*>(FAssetEditorManager::Get().FindEditorForAsset(NewAsset, true));
+		IBehaviorTreeEditor* ChildNodeEditor = static_cast<IBehaviorTreeEditor*>(GEditor->GetEditorSubsystem<UAssetEditorSubsystem>()->FindEditorForAsset(NewAsset, true));
 		if (ChildNodeEditor)
 		{
 			ChildNodeEditor->InitializeDebuggerState(Debugger.Get());
@@ -1761,7 +1762,7 @@ void FBehaviorTreeEditor::HandleNewNodeClassPicked(UClass* InClass) const
 			// Create and init a new Blueprint
 			if (UBlueprint* NewBP = FKismetEditorUtilities::CreateBlueprint(InClass, Package, FName(*Name), BPTYPE_Normal, UBlueprint::StaticClass(), UBlueprintGeneratedClass::StaticClass()))
 			{
-				FAssetEditorManager::Get().OpenEditorForAsset(NewBP);
+				GEditor->GetEditorSubsystem<UAssetEditorSubsystem>()->OpenEditorForAsset(NewBP);
 
 				// Notify the asset registry
 				FAssetRegistryModule::AssetCreated(NewBP);
diff --git a/Engine/Source/Editor/BlueprintGraph/Private/BlueprintEditorSettings.cpp b/Engine/Source/Editor/BlueprintGraph/Private/BlueprintEditorSettings.cpp
--- a/Engine/Source/Editor/BlueprintGraph/Private/BlueprintEditorSettings.cpp
+++ b/Engine/Source/Editor/BlueprintGraph/Private/BlueprintEditorSettings.cpp
@@ -4,8 +4,10 @@
 #include "Misc/ConfigCacheIni.h"
 #include "Editor/EditorPerProjectUserSettings.h"
 #include "Settings/EditorExperimentalSettings.h"
-#include "Toolkits/AssetEditorManager.h"
+
 #include "FindInBlueprintManager.h"
+#include "Subsystems/AssetEditorSubsystem.h"
+#include "Editor.h"
 
 UBlueprintEditorSettings::UBlueprintEditorSettings(const FObjectInitializer& ObjectInitializer)
 	: Super(ObjectInitializer)
@@ -66,13 +68,13 @@ void UBlueprintEditorSettings::PostEditChangeProperty(FPropertyChangedEvent& Pro
 	if (PropertyName == GET_MEMBER_NAME_CHECKED(UBlueprintEditorSettings, bHostFindInBlueprintsInGlobalTab))
 	{
 		// Close all open Blueprint editors to reset associated FiB states.
-		FAssetEditorManager& AssetEditorManager = FAssetEditorManager::Get();
-		TArray<UObject*> EditedAssets = AssetEditorManager.GetAllEditedAssets();
+		UAssetEditorSubsystem* AssetEditorSubsystem = GEditor->GetEditorSubsystem<UAssetEditorSubsystem>();
+		TArray<UObject*> EditedAssets = AssetEditorSubsystem->GetAllEditedAssets();
 		for (UObject* EditedAsset : EditedAssets)
 		{
 			if (EditedAsset->IsA<UBlueprint>())
 			{
-				AssetEditorManager.CloseAllEditorsForAsset(EditedAsset);
+				AssetEditorSubsystem->CloseAllEditorsForAsset(EditedAsset);
 			}
 		}
 
diff --git a/Engine/Source/Editor/Blutility/Private/BlutilityMenuExtensions.cpp b/Engine/Source/Editor/Blutility/Private/BlutilityMenuExtensions.cpp
--- a/Engine/Source/Editor/Blutility/Private/BlutilityMenuExtensions.cpp
+++ b/Engine/Source/Editor/Blutility/Private/BlutilityMenuExtensions.cpp
@@ -4,7 +4,7 @@
 #include "AssetRegistryModule.h"
 #include "EditorUtilityBlueprint.h"
 #include "Misc/PackageName.h"
-#include "Toolkits/AssetEditorManager.h"
+
 #include "BlueprintEditorModule.h"
 #include "GlobalEditorUtilityBase.h"
 #include "Framework/MultiBox/MultiBoxBuilder.h"
@@ -20,6 +20,7 @@
 #include "UObject/PropertyPortFlags.h"
 #include "Widgets/Layout/SScrollBox.h"
 #include "ScopedTransaction.h"
+#include "Subsystems/AssetEditorSubsystem.h"
 
 #define LOCTEXT_NAMESPACE "BlutilityMenuExtensions"
 
@@ -261,7 +262,7 @@ void FBlutilityMenuExtensions::CreateBlutilityActionsMenu(FMenuBuilder& MenuBuil
 								// Edit the script if we have shift held down
 								if(UBlueprint* Blueprint = Cast<UBlueprint>(Cast<UObject>(FunctionAndUtil.Util)->GetClass()->ClassGeneratedBy))
 								{
-									if(IAssetEditorInstance* AssetEditor = FAssetEditorManager::Get().FindEditorForAsset(Blueprint, true))
+									if(IAssetEditorInstance* AssetEditor = GEditor->GetEditorSubsystem<UAssetEditorSubsystem>()->FindEditorForAsset(Blueprint, true))
 									{
 										check(AssetEditor->GetEditorName() == TEXT("BlueprintEditor"));
 										IBlueprintEditor* BlueprintEditor = static_cast<IBlueprintEditor*>(AssetEditor);
diff --git a/Engine/Source/Editor/Cascade/Private/Tests/ParticleSystemTests.cpp b/Engine/Source/Editor/Cascade/Private/Tests/ParticleSystemTests.cpp
--- a/Engine/Source/Editor/Cascade/Private/Tests/ParticleSystemTests.cpp
+++ b/Engine/Source/Editor/Cascade/Private/Tests/ParticleSystemTests.cpp
@@ -6,7 +6,7 @@
 #include "Factories/ParticleSystemFactoryNew.h"
 #include "Particles/ParticleSystem.h"
 #include "Distributions/DistributionVectorUniform.h"
-#include "Toolkits/AssetEditorManager.h"
+
 
 //Automation
 #include "Tests/AutomationTestSettings.h"
@@ -17,6 +17,7 @@
 #include "Particles/ParticleEmitter.h"
 #include "Particles/ParticleLODLevel.h"
 #include "Particles/Size/ParticleModuleSize.h"
+#include "Subsystems/AssetEditorSubsystem.h"
 
 
 #if WITH_DEV_AUTOMATION_TESTS
@@ -158,7 +159,7 @@ namespace ParticleEditorPromotionTestHelper
 			if (CreatedPS)
 			{
 				// @TODO: Add verification that this completed successfully.
-				FAssetEditorManager::Get().OpenEditorForAsset(CreatedPS);
+				GEditor->GetEditorSubsystem<UAssetEditorSubsystem>()->OpenEditorForAsset(CreatedPS);
 				UE_LOG(LogParticleEditorPromotionTests, Display, TEXT("Opened the cascade editor"));
 			}
 			else
@@ -178,7 +179,7 @@ namespace ParticleEditorPromotionTestHelper
 			// @TODO: This should create, edit, and remove a temp asset instead of relying on a previously-created one
 			if (CreatedPS)
 			{
-				IAssetEditorInstance* AssetEditor = FAssetEditorManager::Get().FindEditorForAsset(CreatedPS, true);
+				IAssetEditorInstance* AssetEditor = GEditor->GetEditorSubsystem<UAssetEditorSubsystem>()->FindEditorForAsset(CreatedPS, true);
 
 				bool bModifiedSize = false;
 				UParticleLODLevel* DefaultLOD = CreatedPS->Emitters[0]->LODLevels[0];
@@ -209,7 +210,7 @@ namespace ParticleEditorPromotionTestHelper
 				}
 
 				// @TODO: Close only Cascade instead of all editors
-				FAssetEditorManager::Get().CloseAllAssetEditors();
+				GEditor->GetEditorSubsystem<UAssetEditorSubsystem>()->CloseAllAssetEditors();
 				UE_LOG(LogParticleEditorPromotionTests, Display, TEXT("Closed the cascade editor"));
 			}
 			// @TODO: Return false if test failed after removing sequential-testing setup
diff --git a/Engine/Source/Editor/ClassViewer/Private/SClassViewer.cpp b/Engine/Source/Editor/ClassViewer/Private/SClassViewer.cpp
--- a/Engine/Source/Editor/ClassViewer/Private/SClassViewer.cpp
+++ b/Engine/Source/Editor/ClassViewer/Private/SClassViewer.cpp
@@ -47,7 +47,7 @@
 
 #include "Kismet2/KismetEditorUtilities.h"
 #include "Editor.h"
-#include "Toolkits/AssetEditorManager.h"
+
 #include "PackageTools.h"
 #include "Logging/MessageLog.h"
 
@@ -76,6 +76,7 @@
 #include "SListViewSelectorDropdownMenu.h"
 #include "ClassViewerProjectSettings.h"
 #include "Widgets/Layout/SScrollBorder.h"
+#include "Subsystems/AssetEditorSubsystem.h"
 
 #define LOCTEXT_NAMESPACE "SClassViewer"
 
@@ -505,7 +506,7 @@ namespace ClassViewer
 				UBlueprint* NewBP = FKismetEditorUtilities::CreateBlueprint(InCreationClass, Package, BPName, BPTYPE_Normal, UBlueprint::StaticClass(), UBlueprintGeneratedClass::StaticClass(), FName("ClassViewer"));
 				if(NewBP)
 				{
-					FAssetEditorManager::Get().OpenEditorForAsset(NewBP);
+					GEditor->GetEditorSubsystem<UAssetEditorSubsystem>()->OpenEditorForAsset(NewBP);
 
 					// Notify the asset registry
 					FAssetRegistryModule::AssetCreated(NewBP);
@@ -635,7 +636,7 @@ namespace ClassViewer
 		{
 			if( InBlueprint != nullptr )
 			{
-				FAssetEditorManager::Get().OpenEditorForAsset(InBlueprint);
+				GEditor->GetEditorSubsystem<UAssetEditorSubsystem>()->OpenEditorForAsset(InBlueprint);
 			}
 		}
 
diff --git a/Engine/Source/Editor/ContentBrowser/Private/AssetContextMenu.cpp b/Engine/Source/Editor/ContentBrowser/Private/AssetContextMenu.cpp
--- a/Engine/Source/Editor/ContentBrowser/Private/AssetContextMenu.cpp
+++ b/Engine/Source/Editor/ContentBrowser/Private/AssetContextMenu.cpp
@@ -48,7 +48,7 @@
 #include "ObjectTools.h"
 #include "PackageTools.h"
 #include "Editor.h"
-#include "Toolkits/AssetEditorManager.h"
+
 #include "PropertyEditorModule.h"
 #include "Toolkits/GlobalEditorCommonCommands.h"
 #include "ConsolidateWindow.h"
@@ -77,6 +77,7 @@
 
 #include "PackageHelperFunctions.h"
 #include "EngineUtils.h"
+#include "Subsystems/AssetEditorSubsystem.h"
 
 #define LOCTEXT_NAMESPACE "ContentBrowser"
 
@@ -1195,7 +1196,7 @@ void FAssetContextMenu::ExecuteFindInAssetTree(TArray<FName> InAssets)
 
 void FAssetContextMenu::ExecuteOpenEditorsForAssets(TArray<FName> InAssets)
 {
-	FAssetEditorManager::Get().OpenEditorsForAssets(InAssets);
+	GEditor->GetEditorSubsystem<UAssetEditorSubsystem>()->OpenEditorsForAssets(InAssets);
 }
 
 bool FAssetContextMenu::AddReferenceMenuOptions(UToolMenu* Menu)
@@ -2126,7 +2127,7 @@ void FAssetContextMenu::ExecuteEditAsset()
 	for (const auto& AssetsByClassPair : SelectedAssetsByClass)
 	{
 		const auto& TypeAssets = AssetsByClassPair.Value;
-		FAssetEditorManager::Get().OpenEditorForAssets(TypeAssets);
+		GEditor->GetEditorSubsystem<UAssetEditorSubsystem>()->OpenEditorForAssets(TypeAssets);
 	}
 }
 
diff --git a/Engine/Source/Editor/ContentBrowser/Private/ContentBrowserUtils.cpp b/Engine/Source/Editor/ContentBrowser/Private/ContentBrowserUtils.cpp
--- a/Engine/Source/Editor/ContentBrowser/Private/ContentBrowserUtils.cpp
+++ b/Engine/Source/Editor/ContentBrowser/Private/ContentBrowserUtils.cpp
@@ -43,7 +43,7 @@
 #include "EmptyFolderVisibilityManager.h"
 #include "Settings/EditorExperimentalSettings.h"
 
-#include "Toolkits/AssetEditorManager.h"
+
 #include "PackagesDialog.h"
 #include "PackageTools.h"
 #include "ObjectTools.h"
@@ -56,6 +56,8 @@
 #include "SAssetView.h"
 #include "SPathView.h"
 #include "ContentBrowserLog.h"
+#include "Subsystems/AssetEditorSubsystem.h"
+#include "Editor.h"
 
 #define LOCTEXT_NAMESPACE "ContentBrowser"
 
@@ -308,7 +310,7 @@ bool ContentBrowserUtils::OpenEditorForAsset(UObject* Asset)
 	if( Asset != NULL )
 	{
 		// @todo toolkit minor: Needs world-centric support?
-		return FAssetEditorManager::Get().OpenEditorForAsset(Asset);
+		return GEditor->GetEditorSubsystem<UAssetEditorSubsystem>()->OpenEditorForAsset(Asset);
 	}
 
 	return false;
@@ -322,7 +324,7 @@ bool ContentBrowserUtils::OpenEditorForAsset(const TArray<UObject*>& Assets)
 	}
 	else if ( Assets.Num() > 1 )
 	{
-		return FAssetEditorManager::Get().OpenEditorForAssets(Assets);
+		return GEditor->GetEditorSubsystem<UAssetEditorSubsystem>()->OpenEditorForAssets(Assets);
 	}
 	
 	return false;
diff --git a/Engine/Source/Editor/DetailCustomizations/Private/AmbientSoundDetails.cpp b/Engine/Source/Editor/DetailCustomizations/Private/AmbientSoundDetails.cpp
--- a/Engine/Source/Editor/DetailCustomizations/Private/AmbientSoundDetails.cpp
+++ b/Engine/Source/Editor/DetailCustomizations/Private/AmbientSoundDetails.cpp
@@ -13,7 +13,7 @@
 #include "Sound/SoundBase.h"
 #include "Sound/SoundCue.h"
 #include "Editor.h"
-#include "Toolkits/AssetEditorManager.h"
+
 #include "DetailLayoutBuilder.h"
 #include "DetailWidgetRow.h"
 #include "DetailCategoryBuilder.h"
@@ -24,6 +24,7 @@
 #include "Sound/SoundNodeDelay.h"
 #include "Sound/SoundNodeRandom.h"
 #include "Sound/SoundNodeWavePlayer.h"
+#include "Subsystems/AssetEditorSubsystem.h"
 
 
 #define LOCTEXT_NAMESPACE "AmbientSoundDetails"
@@ -134,7 +135,7 @@ FReply FAmbientSoundDetails::OnEditSoundCueClicked()
 		USoundCue* SoundCue = Cast<USoundCue>(AmbientSound.Get()->GetAudioComponent()->Sound);
 		if (SoundCue)
 		{
-			FAssetEditorManager::Get().OpenEditorForAsset(SoundCue);
+			GEditor->GetEditorSubsystem<UAssetEditorSubsystem>()->OpenEditorForAsset(SoundCue);
 		}
 	}
 
@@ -270,7 +271,7 @@ void FAmbientSoundDetails::CreateNewSoundCue( ESoundCueLayouts Layout )
 
 		SoundCue->LinkGraphNodesFromSoundNodes();
 
-		FAssetEditorManager::Get().OpenEditorForAsset(SoundCue);
+		GEditor->GetEditorSubsystem<UAssetEditorSubsystem>()->OpenEditorForAsset(SoundCue);
 	}
 }
 
diff --git a/Engine/Source/Editor/DetailCustomizations/Private/CurveColorCustomization.cpp b/Engine/Source/Editor/DetailCustomizations/Private/CurveColorCustomization.cpp
--- a/Engine/Source/Editor/DetailCustomizations/Private/CurveColorCustomization.cpp
+++ b/Engine/Source/Editor/DetailCustomizations/Private/CurveColorCustomization.cpp
@@ -12,13 +12,14 @@
 #include "Widgets/Images/SImage.h"
 #include "Widgets/Input/SButton.h"
 #include "Dialogs/Dialogs.h"
-#include "Toolkits/AssetEditorManager.h"
+
 #include "Dialogs/DlgPickAssetPath.h"
 #include "PackageTools.h"
 #include "MiniCurveEditor.h"
 #include "AssetRegistryModule.h"
 #include "SCurveEditor.h"
 #include "Engine/Selection.h"
+#include "Subsystems/AssetEditorSubsystem.h"
 
 #define LOCTEXT_NAMESPACE "CurveColorCustomization"
 
@@ -400,7 +401,7 @@ FReply FCurveColorCustomization::OnCurvePreviewDoubleClick(const FGeometry& InMy
 	{
 		if (RuntimeCurve->ExternalCurve)
 		{
-			FAssetEditorManager::Get().OpenEditorForAsset(RuntimeCurve->ExternalCurve);
+			GEditor->GetEditorSubsystem<UAssetEditorSubsystem>()->OpenEditorForAsset(RuntimeCurve->ExternalCurve);
 		}
 		else
 		{
diff --git a/Engine/Source/Editor/DetailCustomizations/Private/CurveStructCustomization.cpp b/Engine/Source/Editor/DetailCustomizations/Private/CurveStructCustomization.cpp
--- a/Engine/Source/Editor/DetailCustomizations/Private/CurveStructCustomization.cpp
+++ b/Engine/Source/Editor/DetailCustomizations/Private/CurveStructCustomization.cpp
@@ -12,13 +12,14 @@
 #include "Widgets/Images/SImage.h"
 #include "Widgets/Input/SButton.h"
 #include "Dialogs/Dialogs.h"
-#include "Toolkits/AssetEditorManager.h"
+
 #include "Dialogs/DlgPickAssetPath.h"
 #include "PackageTools.h"
 #include "MiniCurveEditor.h"
 #include "AssetRegistryModule.h"
 #include "SCurveEditor.h"
 #include "Engine/Selection.h"
+#include "Subsystems/AssetEditorSubsystem.h"
 
 #define LOCTEXT_NAMESPACE "CurveStructCustomization"
 
@@ -380,7 +381,7 @@ FReply FCurveStructCustomization::OnCurvePreviewDoubleClick( const FGeometry& In
 	{
 		if (RuntimeCurve->ExternalCurve)
 		{
-			FAssetEditorManager::Get().OpenEditorForAsset(RuntimeCurve->ExternalCurve);
+			GEditor->GetEditorSubsystem<UAssetEditorSubsystem>()->OpenEditorForAsset(RuntimeCurve->ExternalCurve);
 		}
 		else
 		{
diff --git a/Engine/Source/Editor/DetailCustomizations/Private/LevelSequenceActorDetails.cpp b/Engine/Source/Editor/DetailCustomizations/Private/LevelSequenceActorDetails.cpp
--- a/Engine/Source/Editor/DetailCustomizations/Private/LevelSequenceActorDetails.cpp
+++ b/Engine/Source/Editor/DetailCustomizations/Private/LevelSequenceActorDetails.cpp
@@ -5,13 +5,14 @@
 #include "Widgets/SBoxPanel.h"
 #include "Widgets/Input/SButton.h"
 #include "Editor.h"
-#include "Toolkits/AssetEditorManager.h"
+
 #include "DetailLayoutBuilder.h"
 #include "DetailWidgetRow.h"
 #include "DetailCategoryBuilder.h"
 #include "IDetailsView.h"
 #include "LevelSequenceActor.h"
 #include "Algo/Transform.h"
+#include "Subsystems/AssetEditorSubsystem.h"
 
 #define LOCTEXT_NAMESPACE "LevelSequenceActorDetails"
 
@@ -152,7 +153,7 @@ FReply FLevelSequenceActorDetails::OnOpenLevelSequenceForActor()
 		UObject* LoadedObject = LevelSequenceActor.Get()->LevelSequence.TryLoad();
 		if (LoadedObject != nullptr)
 		{
-			FAssetEditorManager::Get().OpenEditorForAsset(LoadedObject);
+			GEditor->GetEditorSubsystem<UAssetEditorSubsystem>()->OpenEditorForAsset(LoadedObject);
 		}
 	}
 
diff --git a/Engine/Source/Editor/DetailCustomizations/Private/PostProcessSettingsCustomization.cpp b/Engine/Source/Editor/DetailCustomizations/Private/PostProcessSettingsCustomization.cpp
--- a/Engine/Source/Editor/DetailCustomizations/Private/PostProcessSettingsCustomization.cpp
+++ b/Engine/Source/Editor/DetailCustomizations/Private/PostProcessSettingsCustomization.cpp
@@ -14,7 +14,7 @@
 #include "Widgets/Input/SComboButton.h"
 #include "Materials/Material.h"
 #include "Materials/MaterialInstanceConstant.h"
-#include "Toolkits/AssetEditorManager.h"
+
 #include "IDetailGroup.h"
 #include "IDetailChildrenBuilder.h"
 #include "PropertyCustomizationHelpers.h"
@@ -22,6 +22,8 @@
 #include "Widgets/Layout/SWidgetSwitcher.h"
 #include "DetailCategoryBuilder.h"
 #include "DetailLayoutBuilder.h"
+#include "Subsystems/AssetEditorSubsystem.h"
+#include "Editor.h"
 
 #define LOCTEXT_NAMESPACE "PostProcessSettingsCustomization"
 
@@ -352,7 +354,7 @@ FReply FWeightedBlendableCustomization::JumpToDirectAsset(TSharedPtr<IPropertyHa
 	
 	Value->GetValue(RefObject);
 
-	FAssetEditorManager::Get().OpenEditorForAsset(RefObject);
+	GEditor->GetEditorSubsystem<UAssetEditorSubsystem>()->OpenEditorForAsset(RefObject);
 
 	return FReply::Handled();
 }
diff --git a/Engine/Source/Editor/DetailCustomizations/Private/WorldSettingsDetails.cpp b/Engine/Source/Editor/DetailCustomizations/Private/WorldSettingsDetails.cpp
--- a/Engine/Source/Editor/DetailCustomizations/Private/WorldSettingsDetails.cpp
+++ b/Engine/Source/Editor/DetailCustomizations/Private/WorldSettingsDetails.cpp
@@ -12,7 +12,7 @@
 #include "GameFramework/Actor.h"
 #include "Editor.h"
 #include "Engine/Texture2D.h"
-#include "Toolkits/AssetEditorManager.h"
+
 #include "DetailLayoutBuilder.h"
 #include "GameModeInfoCustomizer.h"
 #include "Settings/EditorExperimentalSettings.h"
@@ -283,7 +283,7 @@ void FLightmapCustomNodeBuilder::ExecuteViewLightmap(FString SelectedLightmapPat
 	UObject* LightMapObject = FindObject<UObject>(NULL, *SelectedLightmapPath);
 	if ( LightMapObject )
 	{
-		FAssetEditorManager::Get().OpenEditorForAsset(LightMapObject);
+		GEditor->GetEditorSubsystem<UAssetEditorSubsystem>()->OpenEditorForAsset(LightMapObject);
 	}
 }
 
diff --git a/Engine/Source/Editor/Documentation/Private/UDNParser.cpp b/Engine/Source/Editor/Documentation/Private/UDNParser.cpp
--- a/Engine/Source/Editor/Documentation/Private/UDNParser.cpp
+++ b/Engine/Source/Editor/Documentation/Private/UDNParser.cpp
@@ -21,11 +21,12 @@
 #include "ISourceCodeAccessModule.h"
 #include "IContentBrowserSingleton.h"
 #include "ContentBrowserModule.h"
-#include "Toolkits/AssetEditorManager.h"
+
 #include "DesktopPlatformModule.h"
 #include "Framework/Notifications/NotificationManager.h"
 #include "Widgets/Notifications/SNotificationList.h"
 #include "Widgets/Input/SHyperlink.h"
+#include "Subsystems/AssetEditorSubsystem.h"
 
 #define LOCTEXT_NAMESPACE "IntroTutorials"
 
@@ -1344,7 +1345,7 @@ bool FUDNParser::ParseAssetLink(FString &InternalLink)
 		{
 			if (Action == TEXT("EDIT"))
 			{
-				FAssetEditorManager::Get().OpenEditorForAsset(RequiredObject);
+				GEditor->GetEditorSubsystem<UAssetEditorSubsystem>()->OpenEditorForAsset(RequiredObject);
 			}
 			else
 			{
diff --git a/Engine/Source/Editor/EnvironmentQueryEditor/Private/SEnvQueryProfiler.cpp b/Engine/Source/Editor/EnvironmentQueryEditor/Private/SEnvQueryProfiler.cpp
--- a/Engine/Source/Editor/EnvironmentQueryEditor/Private/SEnvQueryProfiler.cpp
+++ b/Engine/Source/Editor/EnvironmentQueryEditor/Private/SEnvQueryProfiler.cpp
@@ -15,9 +15,10 @@
 #include "Editor/UnrealEdEngine.h"
 #include "UnrealEdGlobals.h"
 #include "EditorStyleSet.h"
-#include "Toolkits/AssetEditorManager.h"
+
 #include "AssetRegistryModule.h"
 #include "EnvironmentQuery/EnvQuery.h"
+#include "Subsystems/AssetEditorSubsystem.h"
 
 #define LOCTEXT_NAMESPACE "EnvironmentQueryEditor"
 
@@ -230,7 +231,7 @@ void SEnvQueryProfiler::OnItemDoubleClicked(TSharedPtr<SEnvQueryProfiler::FStatD
 	UObject* QueryOb = FindQueryObjectByName(Item->QueryName);
 	if (QueryOb)
 	{
-		FAssetEditorManager::Get().OpenEditorForAsset(QueryOb);
+		GEditor->GetEditorSubsystem<UAssetEditorSubsystem>()->OpenEditorForAsset(QueryOb);
 	}
 }
 
