Remove SurroundSoundFactory, thus deprecating specialized import for multiple wav files with special cased file names.

diff --git a/Engine/Plugins/Tests/EditorTests/Source/EditorTests/Private/UnrealEd/EditorBuildPromotionTests.cpp b/Engine/Plugins/Tests/EditorTests/Source/EditorTests/Private/UnrealEd/EditorBuildPromotionTests.cpp
--- a/Engine/Plugins/Tests/EditorTests/Source/EditorTests/Private/UnrealEd/EditorBuildPromotionTests.cpp
+++ b/Engine/Plugins/Tests/EditorTests/Source/EditorTests/Private/UnrealEd/EditorBuildPromotionTests.cpp
@@ -46,7 +46,6 @@
 #include "Factories/BlueprintFactory.h"
 #include "Factories/FbxFactory.h"
 #include "Factories/SoundFactory.h"
-#include "Factories/SoundSurroundFactory.h"
 #include "Factories/TextureFactory.h"
 #include "Factories/FbxImportUI.h"
 #include "Settings/LevelEditorMiscSettings.h"
@@ -1289,7 +1288,7 @@ namespace BuildPromotionTestHelper
 			{
 				const FString BaseFileName = FPaths::GetPath(SurroundFilePath) / FPaths::GetBaseFilename(SurroundFilePath).LeftChop(3);
 
-				USoundSurroundFactory* FactoryInst = NewObject<USoundSurroundFactory>();
+				USoundFactory* FactoryInst = NewObject<USoundFactory>();
 				FAutomationEditorCommonUtils::ApplyCustomFactorySettings(FactoryInst, AutomationTestSettings->BuildPromotionTest.ImportWorkflow.SurroundSound.FactorySettings);
 
 				const FString SurroundChannels[] = { TEXT("_fl"), TEXT("_fr"), TEXT("_fc"), TEXT("_lf"), TEXT("_sl"), TEXT("_sr"), TEXT("_bl"), TEXT("_br") };
diff --git a/Engine/Source/Editor/AudioEditor/Classes/Factories/ReimportSoundSurroundFactory.h b/Engine/Source/Editor/AudioEditor/Classes/Factories/ReimportSoundSurroundFactory.h
deleted file mode 100644
diff --git a/Engine/Source/Editor/AudioEditor/Classes/Factories/SoundSurroundFactory.h b/Engine/Source/Editor/AudioEditor/Classes/Factories/SoundSurroundFactory.h
deleted file mode 100644
diff --git a/Engine/Source/Editor/AudioEditor/Private/AudioEditorModule.cpp b/Engine/Source/Editor/AudioEditor/Private/AudioEditorModule.cpp
--- a/Engine/Source/Editor/AudioEditor/Private/AudioEditorModule.cpp
+++ b/Engine/Source/Editor/AudioEditor/Private/AudioEditorModule.cpp
@@ -9,7 +9,6 @@
 #include "Factories/ReimportSoundFactory.h"
 #include "SoundCueGraph/SoundCueGraphNode.h"
 #include "SoundCueGraphNodeFactory.h"
-#include "Factories/ReimportSoundSurroundFactory.h"
 #include "AssetToolsModule.h"
 #include "PropertyEditorModule.h"
 #include "SoundClassEditor.h"
@@ -93,9 +92,6 @@ public:
 		// Create reimport handler for sound node waves
 		UReimportSoundFactory::StaticClass();
 
-		// Create reimport handler for surround sound waves
-		UReimportSoundSurroundFactory::StaticClass();
-
 		SetupIcons();
 #if WITH_SNDFILE_IO
 		if (!Audio::InitSoundFileIOManager())
diff --git a/Engine/Source/Editor/AudioEditor/Private/Factories/ReimportSoundSurroundFactory.cpp b/Engine/Source/Editor/AudioEditor/Private/Factories/ReimportSoundSurroundFactory.cpp
deleted file mode 100644
diff --git a/Engine/Source/Editor/AudioEditor/Private/Factories/SoundSurroundFactory.cpp b/Engine/Source/Editor/AudioEditor/Private/Factories/SoundSurroundFactory.cpp
deleted file mode 100644
diff --git a/Engine/Source/Editor/UnrealEd/Private/Editor.cpp b/Engine/Source/Editor/UnrealEd/Private/Editor.cpp
--- a/Engine/Source/Editor/UnrealEd/Private/Editor.cpp
+++ b/Engine/Source/Editor/UnrealEd/Private/Editor.cpp
@@ -16,7 +16,6 @@
 #include "Factories/ReimportFbxStaticMeshFactory.h"
 #include "Factories/ReimportFbxSceneFactory.h"
 #include "Factories/ReimportSoundFactory.h"
-#include "Factories/ReimportSoundSurroundFactory.h"
 #include "Factories/ReimportTextureFactory.h"
 #include "Factories/PhysicalMaterialMaskFactory.h"
 #include "Engine/BlueprintGeneratedClass.h"
diff --git a/Engine/Source/Editor/UnrealEd/Private/Factories/EditorFactories.cpp b/Engine/Source/Editor/UnrealEd/Private/Factories/EditorFactories.cpp
--- a/Engine/Source/Editor/UnrealEd/Private/Factories/EditorFactories.cpp
+++ b/Engine/Source/Editor/UnrealEd/Private/Factories/EditorFactories.cpp
@@ -108,7 +108,6 @@
 #include "Factories/SoundCueFactoryNew.h"
 #include "Factories/ReimportSoundFactory.h"
 #include "Factories/SoundMixFactory.h"
-#include "Factories/ReimportSoundSurroundFactory.h"
 #include "Factories/StructureFactory.h"
 #include "Factories/StringTableFactory.h"
 #include "Factories/SubsurfaceProfileFactory.h"
