EditorSettingsTests: Replace a failed check with an error when the editor can not be focused, to avoid crashing the rest of the tests in the suite.
#rnx

diff --git a/Engine/Plugins/Tests/EditorTests/Source/EditorTests/Private/UnrealEd/EditorSettingsTests.cpp b/Engine/Plugins/Tests/EditorTests/Source/EditorTests/Private/UnrealEd/EditorSettingsTests.cpp
--- a/Engine/Plugins/Tests/EditorTests/Source/EditorTests/Private/UnrealEd/EditorSettingsTests.cpp
+++ b/Engine/Plugins/Tests/EditorTests/Source/EditorTests/Private/UnrealEd/EditorSettingsTests.cpp
@@ -230,20 +230,27 @@ bool FEditorSettingsKeybindingsTest::RunTest(const FString& Parameters)
 	//Focus the main editor
 	TArray< TSharedRef<SWindow> > AllWindows;
 	FSlateApplication::Get().GetAllVisibleWindowsOrdered(AllWindows);
-	FSlateApplication::Get().ProcessWindowActivatedEvent(FWindowActivateEvent(FWindowActivateEvent::EA_Activate, AllWindows[0]));
-
-	// TODO Automation Disabling this code until we understand why this no longer works.  Did it ever?
-	if ( 0 )
+	if (AllWindows.Num() == 0)
+	{
+		// This can happen when the user has activated a different fullscreen application during the test. Mark the test as failed.
+		AddError(TEXT("Can not focus the main editor"));
+	}
+	else
 	{
-		//Send the PIE event
-		FKeyEvent PIEKeyEvent(EKeys::L, FModifierKeysState(false, false, true, false, false, false, false, false, false), false, 0/*UserIndex*/, 0x4C, 0x4C);
-		FSlateApplication::Get().ProcessKeyDownEvent(PIEKeyEvent);
-		FSlateApplication::Get().ProcessKeyUpEvent(PIEKeyEvent);
+		FSlateApplication::Get().ProcessWindowActivatedEvent(FWindowActivateEvent(FWindowActivateEvent::EA_Activate, AllWindows[0]));
+		// TODO Automation Disabling this code until we understand why this no longer works.  Did it ever?
+		if (0)
+		{
+			//Send the PIE event
+			FKeyEvent PIEKeyEvent(EKeys::L, FModifierKeysState(false, false, true, false, false, false, false, false, false), false, 0/*UserIndex*/, 0x4C, 0x4C);
+			FSlateApplication::Get().ProcessKeyDownEvent(PIEKeyEvent);
+			FSlateApplication::Get().ProcessKeyUpEvent(PIEKeyEvent);
 
-		UE_LOG(LogEditorSettingsTests, Display, TEXT("Sent PIE keyboard shortcut"));
+			UE_LOG(LogEditorSettingsTests, Display, TEXT("Sent PIE keyboard shortcut"));
 
-		ADD_LATENT_AUTOMATION_COMMAND(FWaitLatentCommand(3.f));
-		ADD_LATENT_AUTOMATION_COMMAND(FSettingsCheckForPIECommand());
+			ADD_LATENT_AUTOMATION_COMMAND(FWaitLatentCommand(3.f));
+			ADD_LATENT_AUTOMATION_COMMAND(FSettingsCheckForPIECommand());
+		}
 	}
 
 	//Import original keybindings and set changed binds back to cached values
