Add a complex test that Play each map from a project in PIE
Each map is its own test and are listed under Project.PlayMapInPIE

#rnx

diff --git a/Engine/Plugins/Tests/EditorTests/Source/EditorTests/Private/UnrealEd/PlayMapInPIETests.cpp b/Engine/Plugins/Tests/EditorTests/Source/EditorTests/Private/UnrealEd/PlayMapInPIETests.cpp
new file mode 100644
--- /dev/null
+++ b/Engine/Plugins/Tests/EditorTests/Source/EditorTests/Private/UnrealEd/PlayMapInPIETests.cpp
@@ -0,0 +1,211 @@
+// Copyright Epic Games, Inc. All Rights Reserved.
+
+#include "CoreMinimal.h"
+#include "Misc/AutomationTest.h"
+#include "Modules/ModuleManager.h"
+#include "Engine/Engine.h"
+#include "AssetData.h"
+#include "AssetRegistryModule.h"
+#include "EditorTests.h"
+#include "EngineGlobals.h"
+#include "Tests/AutomationCommon.h"
+
+#if WITH_DEV_AUTOMATION_TESTS
+
+#define LOCTEXT_NAMESPACE "PlayMapInPIE"
+
+DEFINE_LOG_CATEGORY_STATIC(LogPlayMapInPIE, Log, All);
+
+DEFINE_LATENT_AUTOMATION_COMMAND_ONE_PARAMETER(FWaitDelay, float, Delay);
+
+bool FWaitDelay::Update()
+{
+	float NewTime = FPlatformTime::Seconds();
+	return (NewTime - StartTime > Delay);
+}
+
+class FPlayMapInPIEBase : public FAutomationTestBase
+{
+public:
+	FPlayMapInPIEBase(const FString& InName, const bool bInComplexTask)
+		: FAutomationTestBase(InName, bInComplexTask)
+	{
+	}
+
+	static void ParseTestMapInfo(const FString& Parameters, FString& MapObjectPath, FString& MapPackageName)
+	{
+		TArray<FString> ParamArray;
+		Parameters.ParseIntoArray(ParamArray, TEXT(";"), true);
+
+		MapObjectPath = ParamArray[0];
+		MapPackageName = ParamArray[1];
+	}
+
+	static UWorld* GetAnyGameWorld()
+	{
+		UWorld* TestWorld = nullptr;
+		const TIndirectArray<FWorldContext>& WorldContexts = GEngine->GetWorldContexts();
+		for (const FWorldContext& Context : WorldContexts)
+		{
+			if (((Context.WorldType == EWorldType::PIE) || (Context.WorldType == EWorldType::Game)) && (Context.World() != NULL))
+			{
+				TestWorld = Context.World();
+				break;
+			}
+		}
+
+		return TestWorld;
+	}
+
+	/** 
+	 * Requests a enumeration of all maps to be loaded
+	 */
+	virtual void GetTests(TArray<FString>& OutBeautifiedNames, TArray<FString>& OutTestCommands) const override
+	{
+		IAssetRegistry& AssetRegistry = FModuleManager::Get().LoadModuleChecked<FAssetRegistryModule>("AssetRegistry").Get();
+
+		if (!AssetRegistry.IsLoadingAssets())
+		{
+#if WITH_EDITOR
+			static bool bDidScan = false;
+
+			if (!GIsEditor && !bDidScan)
+			{
+				// For editor build -game, we need to do a full scan
+				AssetRegistry.SearchAllAssets(true);
+				bDidScan = true;
+			}
+#endif
+
+			TArray<FAssetData> MapList;
+			FARFilter Filter;
+			Filter.ClassNames.Add(UWorld::StaticClass()->GetFName());
+			Filter.bRecursiveClasses = true;
+			Filter.bIncludeOnlyOnDiskAssets = true;
+			if (AssetRegistry.GetAssets(Filter, /*out*/ MapList))
+			{
+				for (const FAssetData& MapAsset : MapList)
+				{
+					FString MapPackageName = MapAsset.PackageName.ToString();
+					if (MapPackageName.Find(TEXT("/Game/")) == 0)
+					{
+						FString MapAssetPath = MapAsset.ObjectPath.ToString();
+
+						OutBeautifiedNames.Add(MapPackageName.RightChop(6)); // Remove "/Game/" from the name
+						OutTestCommands.Add(MapAssetPath + TEXT(";") + MapPackageName);
+					}
+				}
+			}
+		}
+	}
+
+	virtual FString GetTestOpenCommand(const FString& Parameters) const override
+	{
+		FString MapObjectPath, MapPackageName;
+		ParseTestMapInfo(Parameters, MapObjectPath, MapPackageName);
+
+		return FString::Printf(TEXT("Automate.OpenMap %s"), *MapObjectPath);
+	}
+
+	virtual FString GetTestAssetPath(const FString& Parameters) const override
+	{
+		FString MapObjectPath, MapPackageName;
+		ParseTestMapInfo(Parameters, MapObjectPath, MapPackageName);
+
+		return MapObjectPath;
+	}
+
+	/**
+	 * Execute the loading of each map and performance captures
+	 *
+	 * @param Parameters - Should specify which map name to load
+	 * @return	TRUE if the test was successful, FALSE otherwise
+	 */
+	virtual bool RunTest(const FString& Parameters) override
+	{
+		FString MapObjectPath, MapPackageName;
+		ParseTestMapInfo(Parameters, MapObjectPath, MapPackageName);
+
+		bool bCanProceed = false;
+
+		UWorld* TestWorld = GetAnyGameWorld();
+		if (TestWorld && TestWorld->GetMapName() == MapPackageName)
+		{
+			// Map is already loaded.
+			bCanProceed = true;
+		}
+		else
+		{
+			bCanProceed = AutomationOpenMap(MapPackageName);
+		}
+
+		if (bCanProceed)
+		{
+			ADD_LATENT_AUTOMATION_COMMAND(FWaitDelay(4.0));
+			return true;
+		}
+
+		UE_LOG(LogPlayMapInPIE, Error, TEXT("Failed to start the %s map (possibly due to BP compilation issues)"), *MapPackageName);
+		return false;
+	}
+
+};
+
+// Editor only tests
+IMPLEMENT_CUSTOM_COMPLEX_AUTOMATION_TEST(FPlayMapInPIE, FPlayMapInPIEBase, "Project.PlayMapInPIE", (EAutomationTestFlags::EditorContext | EAutomationTestFlags::ProductFilter))
+
+void FPlayMapInPIE::GetTests(TArray<FString>& OutBeautifiedNames, TArray<FString>& OutTestCommands) const
+{
+	FPlayMapInPIEBase::GetTests(OutBeautifiedNames, OutTestCommands);
+}
+
+bool FPlayMapInPIE::RunTest(const FString& Parameters)
+{
+	return FPlayMapInPIEBase::RunTest(Parameters);
+}
+
+void OpenMap(const TArray<FString>& Args)
+{
+	if (Args.Num() != 1)
+	{
+		UE_LOG(LogConsoleResponse, Display, TEXT("Automate.OpenMap failed, the number of arguments is wrong.  Automate.OpenMap MapObjectPath\n"));
+		return;
+	}
+
+	FString AssetPath(Args[0]);
+
+	FAssetRegistryModule& AssetRegistryModule = FModuleManager::LoadModuleChecked<FAssetRegistryModule>(TEXT("AssetRegistry"));
+
+	FAssetData MapAssetData = AssetRegistryModule.Get().GetAssetByObjectPath(*AssetPath);
+
+	if (MapAssetData.IsValid())
+	{
+		bool bIsWorldAlreadyOpened = false;
+		if (UWorld* EditorWorld = GEditor->GetEditorWorldContext().World())
+		{
+			if (FAssetData(EditorWorld).PackageName == MapAssetData.PackageName)
+			{
+				bIsWorldAlreadyOpened = true;
+			}
+		}
+
+		if (!bIsWorldAlreadyOpened)
+		{
+			UObject* ObjectToEdit = MapAssetData.GetAsset();
+			if (ObjectToEdit)
+			{
+				GEditor->EditObject(ObjectToEdit);
+			}
+		}
+	}
+}
+
+FAutoConsoleCommand OpenMapAndFocusActorCmd(
+	TEXT("Automate.OpenMap"),
+	TEXT("Opens a map."),
+	FConsoleCommandWithArgsDelegate::CreateStatic(OpenMap)
+);
+
+#undef LOCTEXT_NAMESPACE
+
+#endif //WITH_DEV_AUTOMATION_TESTS
