Make PlayMapInPIE wait for shader compilation to go through.
Consider Dev Content filter when building the list of tests

#rnx

diff --git a/Engine/Plugins/Tests/EditorTests/Source/EditorTests/EditorTests.Build.cs b/Engine/Plugins/Tests/EditorTests/Source/EditorTests/EditorTests.Build.cs
--- a/Engine/Plugins/Tests/EditorTests/Source/EditorTests/EditorTests.Build.cs
+++ b/Engine/Plugins/Tests/EditorTests/Source/EditorTests/EditorTests.Build.cs
@@ -68,6 +68,7 @@ public class EditorTests : ModuleRules
 				"StaticMeshDescription",
                 "MeshBuilder",
                 "RawMesh",
+				"FunctionalTesting",
 			}
 		);
 
diff --git a/Engine/Plugins/Tests/EditorTests/Source/EditorTests/Private/UnrealEd/PlayMapInPIETests.cpp b/Engine/Plugins/Tests/EditorTests/Source/EditorTests/Private/UnrealEd/PlayMapInPIETests.cpp
--- a/Engine/Plugins/Tests/EditorTests/Source/EditorTests/Private/UnrealEd/PlayMapInPIETests.cpp
+++ b/Engine/Plugins/Tests/EditorTests/Source/EditorTests/Private/UnrealEd/PlayMapInPIETests.cpp
@@ -9,6 +9,9 @@
 #include "EditorTests.h"
 #include "EngineGlobals.h"
 #include "Tests/AutomationCommon.h"
+#include "AutomationBlueprintFunctionLibrary.h"
+#include "NavigationSystem.h"
+#include "IAutomationControllerModule.h"
 
 #if WITH_DEV_AUTOMATION_TESTS
 
@@ -16,6 +19,36 @@
 
 DEFINE_LOG_CATEGORY_STATIC(LogPlayMapInPIE, Log, All);
 
+UWorld* GetTestWorld()
+{
+#if WITH_EDITOR
+	const TIndirectArray<FWorldContext>& WorldContexts = GEngine->GetWorldContexts();
+	for (const FWorldContext& Context : WorldContexts)
+	{
+		if (Context.World() != nullptr)
+		{
+			if (Context.WorldType == EWorldType::PIE)
+			{
+				return Context.World();
+			}
+
+			if (Context.WorldType == EWorldType::Game)
+			{
+				return Context.World();
+			}
+		}
+	}
+#endif
+
+	UWorld* TestWorld = GWorld;
+	if (GIsEditor)
+	{
+		UE_LOG(LogPlayMapInPIE, Warning, TEXT("Play Map In PIE using GWorld. Not correct for PIE"));
+	}
+
+	return TestWorld;
+}
+
 DEFINE_LATENT_AUTOMATION_COMMAND_ONE_PARAMETER(FWaitDelay, float, Delay);
 
 bool FWaitDelay::Update()
@@ -24,6 +57,16 @@ bool FWaitDelay::Update()
 	return (NewTime - StartTime > Delay);
 }
 
+DEFINE_LATENT_AUTOMATION_COMMAND(FWaitForActorsInitialized);
+
+bool FWaitForActorsInitialized::Update()
+{
+	UWorld* World = GetTestWorld();
+
+	return World && World->AreActorsInitialized() &&
+		!UNavigationSystemV1::IsNavigationBeingBuilt(World);
+}
+
 class FPlayMapInPIEBase : public FAutomationTestBase
 {
 public:
@@ -76,6 +119,9 @@ public:
 				bDidScan = true;
 			}
 #endif
+			IAutomationControllerModule& AutomationControllerModule = FModuleManager::LoadModuleChecked<IAutomationControllerModule>(TEXT("AutomationController"));
+			IAutomationControllerManagerPtr AutomationController = AutomationControllerModule.GetAutomationController();
+			bool IsDeveloperDirectoryIncluded = AutomationController->IsDeveloperDirectoryIncluded();
 
 			TArray<FAssetData> MapList;
 			FARFilter Filter;
@@ -89,6 +135,8 @@ public:
 					FString MapPackageName = MapAsset.PackageName.ToString();
 					if (MapPackageName.Find(TEXT("/Game/")) == 0)
 					{
+						if (!IsDeveloperDirectoryIncluded && MapPackageName.Find(TEXT("/Game/Developers")) == 0) continue;
+
 						FString MapAssetPath = MapAsset.ObjectPath.ToString();
 
 						OutBeautifiedNames.Add(MapPackageName.RightChop(6)); // Remove "/Game/" from the name
@@ -141,6 +189,10 @@ public:
 
 		if (bCanProceed)
 		{
+			GEngine->ForceGarbageCollection(true);
+			ADD_LATENT_AUTOMATION_COMMAND(FWaitDelay(0.25));
+			UAutomationBlueprintFunctionLibrary::FinishLoadingBeforeScreenshot();
+			ADD_LATENT_AUTOMATION_COMMAND(FWaitForActorsInitialized);
 			ADD_LATENT_AUTOMATION_COMMAND(FWaitDelay(4.0));
 			return true;
 		}
