Cleaning up code in PlayMapInPIE test

#rnx

diff --git a/Engine/Plugins/Tests/EditorTests/Source/EditorTests/Private/UnrealEd/PlayMapInPIETests.cpp b/Engine/Plugins/Tests/EditorTests/Source/EditorTests/Private/UnrealEd/PlayMapInPIETests.cpp
--- a/Engine/Plugins/Tests/EditorTests/Source/EditorTests/Private/UnrealEd/PlayMapInPIETests.cpp
+++ b/Engine/Plugins/Tests/EditorTests/Source/EditorTests/Private/UnrealEd/PlayMapInPIETests.cpp
@@ -6,7 +6,6 @@
 #include "Engine/Engine.h"
 #include "AssetData.h"
 #include "AssetRegistryModule.h"
-#include "EditorTests.h"
 #include "EngineGlobals.h"
 #include "Tests/AutomationCommon.h"
 #include "AutomationBlueprintFunctionLibrary.h"
@@ -19,32 +18,18 @@
 
 DEFINE_LOG_CATEGORY_STATIC(LogPlayMapInPIE, Log, All);
 
-UWorld* GetTestWorld()
+static UWorld* GetAnyGameWorld()
 {
-#if WITH_EDITOR
+	UWorld* TestWorld = nullptr;
 	const TIndirectArray<FWorldContext>& WorldContexts = GEngine->GetWorldContexts();
 	for (const FWorldContext& Context : WorldContexts)
 	{
-		if (Context.World() != nullptr)
+		if (((Context.WorldType == EWorldType::PIE) || (Context.WorldType == EWorldType::Game)) && (Context.World() != NULL))
 		{
-			if (Context.WorldType == EWorldType::PIE)
-			{
-				return Context.World();
-			}
-
-			if (Context.WorldType == EWorldType::Game)
-			{
-				return Context.World();
-			}
+			TestWorld = Context.World();
+			break;
 		}
 	}
-#endif
-
-	UWorld* TestWorld = GWorld;
-	if (GIsEditor)
-	{
-		UE_LOG(LogPlayMapInPIE, Warning, TEXT("Play Map In PIE using GWorld. Not correct for PIE"));
-	}
 
 	return TestWorld;
 }
@@ -61,7 +46,7 @@ DEFINE_LATENT_AUTOMATION_COMMAND(FWaitForActorsInitialized);
 
 bool FWaitForActorsInitialized::Update()
 {
-	UWorld* World = GetTestWorld();
+	UWorld* World = GetAnyGameWorld();
 
 	return World && World->AreActorsInitialized() &&
 		!UNavigationSystemV1::IsNavigationBeingBuilt(World);
@@ -84,22 +69,6 @@ public:
 		MapPackageName = ParamArray[1];
 	}
 
-	static UWorld* GetAnyGameWorld()
-	{
-		UWorld* TestWorld = nullptr;
-		const TIndirectArray<FWorldContext>& WorldContexts = GEngine->GetWorldContexts();
-		for (const FWorldContext& Context : WorldContexts)
-		{
-			if (((Context.WorldType == EWorldType::PIE) || (Context.WorldType == EWorldType::Game)) && (Context.World() != NULL))
-			{
-				TestWorld = Context.World();
-				break;
-			}
-		}
-
-		return TestWorld;
-	}
-
 	/** 
 	 * Requests a enumeration of all maps to be loaded
 	 */
