Reassigned PlayMapInPIETests to Content suite

#rnx

diff --git a/Engine/Plugins/Tests/EditorTests/Source/EditorTests/Private/UnrealEd/PlayMapInPIETests.cpp b/Engine/Plugins/Tests/EditorTests/Source/EditorTests/Private/UnrealEd/PlayMapInPIETests.cpp
--- a/Engine/Plugins/Tests/EditorTests/Source/EditorTests/Private/UnrealEd/PlayMapInPIETests.cpp
+++ b/Engine/Plugins/Tests/EditorTests/Source/EditorTests/Private/UnrealEd/PlayMapInPIETests.cpp
@@ -173,7 +173,7 @@ public:
 };
 
 // Editor only tests
-IMPLEMENT_CUSTOM_COMPLEX_AUTOMATION_TEST(FPlayMapInPIE, FPlayMapInPIEBase, "Project.PlayMapInPIE", (EAutomationTestFlags::EditorContext | EAutomationTestFlags::ProductFilter))
+IMPLEMENT_CUSTOM_COMPLEX_AUTOMATION_TEST(FPlayMapInPIE, FPlayMapInPIEBase, "Content.PlayMapInPIE", (EAutomationTestFlags::EditorContext | EAutomationTestFlags::ProductFilter))
 
 void FPlayMapInPIE::GetTests(TArray<FString>& OutBeautifiedNames, TArray<FString>& OutTestCommands) const
 {
