FBX automation test crashes when import Animation FBXs
#fix Ensure that we empty the transaction buffer, as this holds onto the UAnimDataModel (removing from root does not remove it from there)

diff --git a/Engine/Plugins/Tests/EditorTests/Source/EditorTests/Private/UnrealEd/FbxAutomationTests.cpp b/Engine/Plugins/Tests/EditorTests/Source/EditorTests/Private/UnrealEd/FbxAutomationTests.cpp
--- a/Engine/Plugins/Tests/EditorTests/Source/EditorTests/Private/UnrealEd/FbxAutomationTests.cpp
+++ b/Engine/Plugins/Tests/EditorTests/Source/EditorTests/Private/UnrealEd/FbxAutomationTests.cpp
@@ -225,6 +225,12 @@ bool FFbxImportAssetsAutomationTest::RunTest(const FString& Parameters)
 							}
 						}
 					}
+
+					if ((GEditor != nullptr) && (GEditor->Trans != nullptr))
+					{
+						GEditor->Trans->Reset(FText::FromString("Discard undo history during FBX Automation testing."));
+					}
+
 					CollectGarbage(GARBAGE_COLLECTION_KEEPFLAGS);
 					
 					ImportedObjects.Empty();
@@ -1613,6 +1619,11 @@ bool FFbxImportAssetsAutomationTest::RunTest(const FString& Parameters)
 		}
 		if (TestPlan->bDeleteFolderAssets || TestPlan->Action == EFBXTestPlanActionType::ImportReload)
 		{
+			if ((GEditor != nullptr) && (GEditor->Trans != nullptr))
+			{
+				GEditor->Trans->Reset(FText::FromString("Discard undo history during FBX Automation testing."));
+			}
+
 			//When doing an import-reload we have to destroy the package since it was save
 			//But when we just have everything in memory a garbage collection pass is enough to
 			//delete assets.
