PlayMapInPIE test - Remove dependency to Functional tests

#rnx

diff --git a/Engine/Plugins/Tests/EditorTests/Source/EditorTests/EditorTests.Build.cs b/Engine/Plugins/Tests/EditorTests/Source/EditorTests/EditorTests.Build.cs
--- a/Engine/Plugins/Tests/EditorTests/Source/EditorTests/EditorTests.Build.cs
+++ b/Engine/Plugins/Tests/EditorTests/Source/EditorTests/EditorTests.Build.cs
@@ -68,7 +68,6 @@ public class EditorTests : ModuleRules
 				"StaticMeshDescription",
                 "MeshBuilder",
                 "RawMesh",
-				"FunctionalTesting",
 				"AutomationController",
 			}
 		);
diff --git a/Engine/Plugins/Tests/EditorTests/Source/EditorTests/Private/UnrealEd/PlayMapInPIETests.cpp b/Engine/Plugins/Tests/EditorTests/Source/EditorTests/Private/UnrealEd/PlayMapInPIETests.cpp
--- a/Engine/Plugins/Tests/EditorTests/Source/EditorTests/Private/UnrealEd/PlayMapInPIETests.cpp
+++ b/Engine/Plugins/Tests/EditorTests/Source/EditorTests/Private/UnrealEd/PlayMapInPIETests.cpp
@@ -8,7 +8,7 @@
 #include "AssetRegistryModule.h"
 #include "EngineGlobals.h"
 #include "Tests/AutomationCommon.h"
-#include "AutomationBlueprintFunctionLibrary.h"
+#include "Tests/AutomationEditorCommon.h"
 #include "NavigationSystem.h"
 #include "IAutomationControllerModule.h"
 
@@ -132,6 +132,22 @@ public:
 		return MapObjectPath;
 	}
 
+	static void LoadAllTexture()
+	{
+		FlushAsyncLoading();
+
+		// Make sure we finish all level streaming
+		if (UWorld* GameWorld = GetAnyGameWorld())
+		{
+			GameWorld->FlushLevelStreaming(EFlushLevelStreamingType::Full);
+		}
+
+		// Force all mip maps to load.
+		UTexture::ForceUpdateTextureStreaming();
+
+		IStreamingManager::Get().StreamAllResources(0.0f);
+	}
+
 	/**
 	 * Execute the loading of each map and performance captures
 	 *
@@ -160,7 +176,9 @@ public:
 		{
 			GEngine->ForceGarbageCollection(true);
 			ADD_LATENT_AUTOMATION_COMMAND(FWaitDelay(0.25));
-			UAutomationBlueprintFunctionLibrary::FinishLoadingBeforeScreenshot();
+			ADD_LATENT_AUTOMATION_COMMAND(FWaitForShadersToFinishCompiling());
+			FModuleManager::GetModuleChecked<IAutomationControllerModule>("AutomationController").GetAutomationController()->ResetAutomationTestTimeout(TEXT("shader compilation"));
+			LoadAllTexture();
 			ADD_LATENT_AUTOMATION_COMMAND(FWaitForActorsInitialized);
 			ADD_LATENT_AUTOMATION_COMMAND(FWaitDelay(4.0));
 			return true;
