PlayMapInPIE test  - Fixing issue with loading some map
Don't wait for navigation system

#rnx

diff --git a/Engine/Plugins/Tests/EditorTests/Source/EditorTests/Private/UnrealEd/PlayMapInPIETests.cpp b/Engine/Plugins/Tests/EditorTests/Source/EditorTests/Private/UnrealEd/PlayMapInPIETests.cpp
--- a/Engine/Plugins/Tests/EditorTests/Source/EditorTests/Private/UnrealEd/PlayMapInPIETests.cpp
+++ b/Engine/Plugins/Tests/EditorTests/Source/EditorTests/Private/UnrealEd/PlayMapInPIETests.cpp
@@ -48,8 +48,7 @@ bool FWaitForActorsInitialized::Update()
 {
 	UWorld* World = GetAnyGameWorld();
 
-	return World && World->AreActorsInitialized() &&
-		!UNavigationSystemV1::IsNavigationBeingBuilt(World);
+	return World && World->AreActorsInitialized();
 }
 
 class FPlayMapInPIEBase : public FAutomationTestBase
