Fix CIS no unity for fbx automation test

diff --git a/Engine/Plugins/Tests/EditorTests/Source/EditorTests/Private/UnrealEd/FbxAutomationTests.cpp b/Engine/Plugins/Tests/EditorTests/Source/EditorTests/Private/UnrealEd/FbxAutomationTests.cpp
--- a/Engine/Plugins/Tests/EditorTests/Source/EditorTests/Private/UnrealEd/FbxAutomationTests.cpp
+++ b/Engine/Plugins/Tests/EditorTests/Source/EditorTests/Private/UnrealEd/FbxAutomationTests.cpp
@@ -34,6 +34,9 @@
 #include "FbxMeshUtils.h"
 #include "Tests/FbxAutomationCommon.h"
 
+#include "Editor/EditorEngine.h"
+extern UNREALED_API UEditorEngine* GEditor;
+
 //////////////////////////////////////////////////////////////////////////
 
 /**
