QAGame config - Set Automation Worker timeout from 5 min to 15 min
+ remove useless call in PlayMapInPIETests

#rnx

diff --git a/Engine/Plugins/Tests/EditorTests/Source/EditorTests/Private/UnrealEd/PlayMapInPIETests.cpp b/Engine/Plugins/Tests/EditorTests/Source/EditorTests/Private/UnrealEd/PlayMapInPIETests.cpp
--- a/Engine/Plugins/Tests/EditorTests/Source/EditorTests/Private/UnrealEd/PlayMapInPIETests.cpp
+++ b/Engine/Plugins/Tests/EditorTests/Source/EditorTests/Private/UnrealEd/PlayMapInPIETests.cpp
@@ -42,15 +42,6 @@ bool FWaitDelay::Update()
 	return (NewTime - StartTime > Delay);
 }
 
-DEFINE_LATENT_AUTOMATION_COMMAND(FWaitForActorsInitialized);
-
-bool FWaitForActorsInitialized::Update()
-{
-	UWorld* World = GetAnyGameWorld();
-
-	return World && World->AreActorsInitialized();
-}
-
 class FPlayMapInPIEBase : public FAutomationTestBase
 {
 public:
@@ -178,7 +169,6 @@ public:
 			ADD_LATENT_AUTOMATION_COMMAND(FWaitForShadersToFinishCompiling());
 			FModuleManager::GetModuleChecked<IAutomationControllerModule>("AutomationController").GetAutomationController()->ResetAutomationTestTimeout(TEXT("shader compilation"));
 			LoadAllTexture();
-			ADD_LATENT_AUTOMATION_COMMAND(FWaitForActorsInitialized);
 			ADD_LATENT_AUTOMATION_COMMAND(FWaitDelay(4.0));
 			return true;
 		}
