Updated MeshDescription tests now that RawMesh is always expected to be invalid once a static mesh is loaded.

diff --git a/Engine/Plugins/Tests/EditorTests/Source/EditorTests/Private/UnrealEd/MeshDescriptionTests.cpp b/Engine/Plugins/Tests/EditorTests/Source/EditorTests/Private/UnrealEd/MeshDescriptionTests.cpp
--- a/Engine/Plugins/Tests/EditorTests/Source/EditorTests/Private/UnrealEd/MeshDescriptionTests.cpp
+++ b/Engine/Plugins/Tests/EditorTests/Source/EditorTests/Private/UnrealEd/MeshDescriptionTests.cpp
@@ -503,7 +503,7 @@ bool FMeshDescriptionTest::ConversionTest(FAutomationTestExecutionInfo& Executio
 			//RawMesh to MeshDescription to RawMesh
 			for (int32 LodIndex = 0; LodIndex < AssetMesh->GetNumSourceModels(); ++LodIndex)
 			{
-				if (AssetMesh->GetSourceModel(LodIndex).RawMeshBulkData->IsEmpty())
+				if (!AssetMesh->GetSourceModel(LodIndex).IsMeshDescriptionValid())
 				{
 					check(LodIndex != 0);
 					continue;
