Fix missing include.
#rnx

diff --git a/Engine/Plugins/Tests/EditorTests/Source/EditorTests/Private/UnrealEd/PlayMapInPIETests.cpp b/Engine/Plugins/Tests/EditorTests/Source/EditorTests/Private/UnrealEd/PlayMapInPIETests.cpp
--- a/Engine/Plugins/Tests/EditorTests/Source/EditorTests/Private/UnrealEd/PlayMapInPIETests.cpp
+++ b/Engine/Plugins/Tests/EditorTests/Source/EditorTests/Private/UnrealEd/PlayMapInPIETests.cpp
@@ -6,6 +6,7 @@
 #include "Engine/Engine.h"
 #include "AssetData.h"
 #include "AssetRegistryModule.h"
+#include "ContentStreaming.h"
 #include "EngineGlobals.h"
 #include "Tests/AutomationCommon.h"
 #include "Tests/AutomationEditorCommon.h"
