QAGame - PlayMapInPIETests - blacklist map using InteractiveLoadMode
+ more clean up in PlayMapInPIETests

#rnx

diff --git a/Engine/Plugins/Tests/EditorTests/Source/EditorTests/Private/UnrealEd/PlayMapInPIETests.cpp b/Engine/Plugins/Tests/EditorTests/Source/EditorTests/Private/UnrealEd/PlayMapInPIETests.cpp
--- a/Engine/Plugins/Tests/EditorTests/Source/EditorTests/Private/UnrealEd/PlayMapInPIETests.cpp
+++ b/Engine/Plugins/Tests/EditorTests/Source/EditorTests/Private/UnrealEd/PlayMapInPIETests.cpp
@@ -35,14 +35,6 @@ static UWorld* GetAnyGameWorld()
 	return TestWorld;
 }
 
-DEFINE_LATENT_AUTOMATION_COMMAND_ONE_PARAMETER(FWaitDelay, float, Delay);
-
-bool FWaitDelay::Update()
-{
-	float NewTime = FPlatformTime::Seconds();
-	return (NewTime - StartTime > Delay);
-}
-
 class FPlayMapInPIEBase : public FAutomationTestBase
 {
 public:
@@ -165,12 +157,11 @@ public:
 
 		if (bCanProceed)
 		{
-			GEngine->ForceGarbageCollection(true);
-			ADD_LATENT_AUTOMATION_COMMAND(FWaitDelay(0.25));
-			ADD_LATENT_AUTOMATION_COMMAND(FWaitForShadersToFinishCompiling());
-			FModuleManager::GetModuleChecked<IAutomationControllerModule>("AutomationController").GetAutomationController()->ResetAutomationTestTimeout(TEXT("shader compilation"));
 			LoadAllTexture();
-			ADD_LATENT_AUTOMATION_COMMAND(FWaitDelay(4.0));
+			ADD_LATENT_AUTOMATION_COMMAND(FWaitLatentCommand(0.25));
+			ADD_LATENT_AUTOMATION_COMMAND(FWaitForShadersToFinishCompiling());
+			ADD_LATENT_AUTOMATION_COMMAND(FWaitLatentCommand(4.0));
+
 			return true;
 		}
 
