PlayMapInPIETests - override SetTestContext() to have proper asset name associate to the test

#rnx

diff --git a/Engine/Plugins/Tests/EditorTests/Source/EditorTests/Private/UnrealEd/PlayMapInPIETests.cpp b/Engine/Plugins/Tests/EditorTests/Source/EditorTests/Private/UnrealEd/PlayMapInPIETests.cpp
--- a/Engine/Plugins/Tests/EditorTests/Source/EditorTests/Private/UnrealEd/PlayMapInPIETests.cpp
+++ b/Engine/Plugins/Tests/EditorTests/Source/EditorTests/Private/UnrealEd/PlayMapInPIETests.cpp
@@ -35,6 +35,26 @@ static UWorld* GetAnyGameWorld()
 	return TestWorld;
 }
 
+DEFINE_LATENT_AUTOMATION_COMMAND(FLoadAllTextureLatentCommand);
+
+bool FLoadAllTextureLatentCommand::Update()
+{
+	FlushAsyncLoading();
+
+	// Make sure we finish all level streaming
+	if (UWorld* GameWorld = GetAnyGameWorld())
+	{
+		GameWorld->FlushLevelStreaming(EFlushLevelStreamingType::Full);
+	}
+
+	// Force all mip maps to load.
+	UTexture::ForceUpdateTextureStreaming();
+
+	IStreamingManager::Get().StreamAllResources(0.0f);
+
+	return true;
+}
+
 class FPlayMapInPIEBase : public FAutomationTestBase
 {
 public:
@@ -115,22 +135,6 @@ public:
 		return MapObjectPath;
 	}
 
-	static void LoadAllTexture()
-	{
-		FlushAsyncLoading();
-
-		// Make sure we finish all level streaming
-		if (UWorld* GameWorld = GetAnyGameWorld())
-		{
-			GameWorld->FlushLevelStreaming(EFlushLevelStreamingType::Full);
-		}
-
-		// Force all mip maps to load.
-		UTexture::ForceUpdateTextureStreaming();
-
-		IStreamingManager::Get().StreamAllResources(0.0f);
-	}
-
 	/**
 	 * Execute the loading of each map and performance captures
 	 *
@@ -157,7 +161,7 @@ public:
 
 		if (bCanProceed)
 		{
-			LoadAllTexture();
+			ADD_LATENT_AUTOMATION_COMMAND(FLoadAllTextureLatentCommand());
 			ADD_LATENT_AUTOMATION_COMMAND(FWaitLatentCommand(0.25));
 			ADD_LATENT_AUTOMATION_COMMAND(FWaitForShadersToFinishCompiling());
 			ADD_LATENT_AUTOMATION_COMMAND(FWaitLatentCommand(4.0));
@@ -169,6 +173,14 @@ public:
 		return false;
 	}
 
+protected:
+	virtual void SetTestContext(FString Context) override
+	{
+		FString MapObjectPath, MapPackageName;
+		ParseTestMapInfo(Context, MapObjectPath, MapPackageName);
+		TestParameterContext = MapPackageName.RightChop(6);
+	}
+
 };
 
 // Editor only tests
