Misc code compilation fixes for monolithic editor build

- make both versions of ComputeLinearVoxelIndex() static as they aren't used outside their respective compilation units
- rename one instance of OpenMapAndFocusActorCmd to avoid name collision
- prevent inclusion of differing DX version headers in the editor
  - shader format for the platform was already disabled in the monolithic build case but this is now extended to prevent inclusion of the associated DX headers as well

- testing
  - editor and game compiled locally for specifically impacted platforms
  - Horde preflight:  https://horde.devtools.epicgames.com/job/606dc2e4963d84000192e7e5

diff --git a/Engine/Plugins/Tests/EditorTests/Source/EditorTests/Private/UnrealEd/PlayMapInPIETests.cpp b/Engine/Plugins/Tests/EditorTests/Source/EditorTests/Private/UnrealEd/PlayMapInPIETests.cpp
--- a/Engine/Plugins/Tests/EditorTests/Source/EditorTests/Private/UnrealEd/PlayMapInPIETests.cpp
+++ b/Engine/Plugins/Tests/EditorTests/Source/EditorTests/Private/UnrealEd/PlayMapInPIETests.cpp
@@ -232,7 +232,7 @@ void OpenMap(const TArray<FString>& Args)
 	}
 }
 
-FAutoConsoleCommand OpenMapAndFocusActorCmd(
+FAutoConsoleCommand OpenMapAndFocusActorPIETestsCmd(
 	TEXT("Automate.OpenMap"),
 	TEXT("Opens a map."),
 	FConsoleCommandWithArgsDelegate::CreateStatic(OpenMap)
