Add PIE start and end to Editor Iteration Tests (Iteration Profile's test automation plugin routine).

diff --git a/Engine/Plugins/Tests/EditorTests/Source/EditorTests/Private/UnrealEd/EditorIterationTests.cpp b/Engine/Plugins/Tests/EditorTests/Source/EditorTests/Private/UnrealEd/EditorIterationTests.cpp
--- a/Engine/Plugins/Tests/EditorTests/Source/EditorTests/Private/UnrealEd/EditorIterationTests.cpp
+++ b/Engine/Plugins/Tests/EditorTests/Source/EditorTests/Private/UnrealEd/EditorIterationTests.cpp
@@ -66,5 +66,10 @@ bool FIterationOpenAssets::RunTest(const FString& LongAssetPath)
 		CollectGarbage(GARBAGE_COLLECTION_KEEPFLAGS);
 	}));
 
+	bool bSimulateInEditor = false;
+	AddCommand(new FStartPIECommand(bSimulateInEditor));
+	AddCommand(new FWaitLatentCommand(10.0f));
+	AddCommand(new FEndPlayMapCommand());
+
 	return true;
 }
