Compilation time optimization
  - Disable codegen optims on extremely long function with lots of FORCEINLINE calls in it
  - 36.23s -> 19.24s for FbxAutomationTests.cpp
  - 37.33s -> 19.89 for Module.EditorTests.cpp

#rnx

diff --git a/Engine/Plugins/Tests/EditorTests/Source/EditorTests/Private/UnrealEd/FbxAutomationTests.cpp b/Engine/Plugins/Tests/EditorTests/Source/EditorTests/Private/UnrealEd/FbxAutomationTests.cpp
--- a/Engine/Plugins/Tests/EditorTests/Source/EditorTests/Private/UnrealEd/FbxAutomationTests.cpp
+++ b/Engine/Plugins/Tests/EditorTests/Source/EditorTests/Private/UnrealEd/FbxAutomationTests.cpp
@@ -183,6 +183,10 @@ FString GetFormatedMessageErrorInExpectedResult(FString FileName, FString TestPl
 		*FileName, *TestPlanName, *ExpectedResultName, ExpectedResultIndex);
 }
 
+// Codegen optimization degenerates for very long functions like RunTest when combined with the invokation of lots of FORCEINLINE methods.
+// We don't need this code to be particularly fast anyway. The other way to improve this code would be to split the test in multiple functions.
+BEGIN_FUNCTION_BUILD_OPTIMIZATION
+
 /**
 * Execute the generic import test
 *
@@ -1971,3 +1975,5 @@ bool FFbxImportAssetsAutomationTest::RunTest(const FString& Parameters)
 
 	return CurTestSuccessful;
 }
+
+END_FUNCTION_BUILD_OPTIMIZATION
\ No newline at end of file
