Fix crash running EditorOpenAsset EngineTest by protecting against object becoming invalid

diff --git a/Engine/Plugins/Tests/EditorTests/Source/EditorTests/Private/UnrealEd/EditorOpenAssetAutomationTests.cpp b/Engine/Plugins/Tests/EditorTests/Source/EditorTests/Private/UnrealEd/EditorOpenAssetAutomationTests.cpp
--- a/Engine/Plugins/Tests/EditorTests/Source/EditorTests/Private/UnrealEd/EditorOpenAssetAutomationTests.cpp
+++ b/Engine/Plugins/Tests/EditorTests/Source/EditorTests/Private/UnrealEd/EditorOpenAssetAutomationTests.cpp
@@ -44,8 +44,22 @@ bool FOpenAssetEditors::RunTest(const FString& LongAssetPath)
 
 	AddCommand(new FOpenEditorForAssetCommand(LongAssetPath));
 	AddCommand(new FWaitLatentCommand(0.5f));
-	AddCommand(new FDelayedFunctionLatentCommand([=] {
-		if ( Object->GetOutermost()->IsDirty() )
+	AddCommand(new FDelayedFunctionLatentCommand([WeakObjPtr=TWeakObjectPtr<UObject>(Object),LongAssetPath] {
+		UObject* Obj = WeakObjPtr.Get();
+		if (Obj == nullptr)
+		{		
+			// If the object was destroyed, attempt to find it if it was recreated
+			if (!LongAssetPath.IsEmpty())
+			{
+				Obj = FindObject<UObject>(nullptr, *LongAssetPath);
+			}
+
+			if (Obj == nullptr)
+			{
+				UE_LOG(LogEditorAutomationTests, Error, TEXT("Asset '%s' did not exist when evaluating."), *LongAssetPath);
+			}
+		}
+		if (Obj && Obj->GetOutermost()->IsDirty() )
 		{
 			UE_LOG(LogEditorAutomationTests, Error, TEXT("Asset '%s' was dirty after opening it."), *LongAssetPath);
 		}
diff --git a/Engine/Source/Editor/UnrealEd/Private/Tests/AutomationEditorCommon.cpp b/Engine/Source/Editor/UnrealEd/Private/Tests/AutomationEditorCommon.cpp
--- a/Engine/Source/Editor/UnrealEd/Private/Tests/AutomationEditorCommon.cpp
+++ b/Engine/Source/Editor/UnrealEd/Private/Tests/AutomationEditorCommon.cpp
@@ -737,7 +737,17 @@ bool FOpenEditorForAssetCommand::Update()
 	UObject* Object = StaticLoadObject(UObject::StaticClass(), NULL, *AssetName);
 	if ( Object )
 	{
+		// Some assets (like UWorlds) may be destroyed and recreated as part of opening. To protect against this, keep the path to the asset and try to re-find it if it disappeared.
+		TWeakObjectPtr<UObject> WeakObject = Object;
+
 		GEditor->GetEditorSubsystem<UAssetEditorSubsystem>()->OpenEditorForAsset(Object);
+
+		// If the object was destroyed, attempt to find it if it was recreated
+		if (!WeakObject.IsValid() && !AssetName.IsEmpty())
+		{
+			Object = FindObject<UObject>(nullptr, *AssetName);
+		}
+
 		//This checks to see if the asset sub editor is loaded.
 		if ( GEditor->GetEditorSubsystem<UAssetEditorSubsystem>()->FindEditorForAsset(Object, true) != NULL )
 		{
