Use EngineDir() to resolve asset path for automated tests from engine config file

#rnx

diff --git a/Engine/Config/BaseEngine.ini b/Engine/Config/BaseEngine.ini
--- a/Engine/Config/BaseEngine.ini
+++ b/Engine/Config/BaseEngine.ini
@@ -2187,12 +2187,12 @@ bFilterDependenciesWithNoTags=false
 bFilterSearchableNames=true
 
 [AutomationTesting]
-ImportTestPath=../../Content/EditorAutomation/
+ImportTestPath=Content/EditorAutomation/
 ImportTestPackagePath=/Engine/Content/EditorAutomation
 bForceSmokeTests=false
 
 [AutomationTesting.FbxImport]
-FbxImportTestPath=../../Content/FbxEditorAutomation/
+FbxImportTestPath=Content/FbxEditorAutomation/
 FbxImportTestPackagePath=/Engine/FbxEditorAutomationOut
 
 [AutomationTesting.Blueprint]
diff --git a/Engine/Plugins/Tests/EditorTests/Source/EditorTests/Private/UnrealEd/EditorAutomationTests.cpp b/Engine/Plugins/Tests/EditorTests/Source/EditorTests/Private/UnrealEd/EditorAutomationTests.cpp
--- a/Engine/Plugins/Tests/EditorTests/Source/EditorTests/Private/UnrealEd/EditorAutomationTests.cpp
+++ b/Engine/Plugins/Tests/EditorTests/Source/EditorTests/Private/UnrealEd/EditorAutomationTests.cpp
@@ -132,7 +132,7 @@ void FGenericImportAssetsAutomationTest::GetTests(TArray<FString>& OutBeautified
 	FString ImportTestDirectory;
 	check( GConfig );
 	GConfig->GetString( TEXT("AutomationTesting"), TEXT("ImportTestPath"), ImportTestDirectory, GEngineIni );
-
+	ImportTestDirectory = FPaths::Combine(FPaths::ConvertRelativePathToFull(FPaths::EngineDir()), ImportTestDirectory);
 
 	// Find all files in the GenericImport directory
 	TArray<FString> FilesInDirectory;
diff --git a/Engine/Plugins/Tests/EditorTests/Source/EditorTests/Private/UnrealEd/FbxAutomationTests.cpp b/Engine/Plugins/Tests/EditorTests/Source/EditorTests/Private/UnrealEd/FbxAutomationTests.cpp
--- a/Engine/Plugins/Tests/EditorTests/Source/EditorTests/Private/UnrealEd/FbxAutomationTests.cpp
+++ b/Engine/Plugins/Tests/EditorTests/Source/EditorTests/Private/UnrealEd/FbxAutomationTests.cpp
@@ -27,6 +27,7 @@
 #include "SkinWeightsUtilities.h"
 #include "Animation/DebugSkelMeshComponent.h"
 #include "Rendering/SkeletalMeshModel.h"
+#include "Misc/Paths.h"
 
 #include "Animation/AnimCurveTypes.h"
 #include "Animation/AnimSequence.h"
@@ -131,7 +132,7 @@ void FFbxImportAssetsAutomationTest::GetTests(TArray<FString>& OutBeautifiedName
 	FString ImportTestDirectory;
 	check(GConfig);
 	GConfig->GetString(TEXT("AutomationTesting.FbxImport"), TEXT("FbxImportTestPath"), ImportTestDirectory, GEngineIni);
-
+	ImportTestDirectory = FPaths::Combine(FPaths::ConvertRelativePathToFull(FPaths::EngineDir()), ImportTestDirectory);
 
 	// Find all files in the GenericImport directory
 	TArray<FString> FilesInDirectory;
