IterationProfile: Add LoadAsset and PIE bookmark pairs

These are added alongside the existing LoadAssetAndPIE telemetry measurement, to make diagnosing changes in that value a bit easier in tableau.

diff --git a/Engine/Plugins/Tests/EditorTests/Source/EditorTests/Private/UnrealEd/EditorIterationTests.cpp b/Engine/Plugins/Tests/EditorTests/Source/EditorTests/Private/UnrealEd/EditorIterationTests.cpp
--- a/Engine/Plugins/Tests/EditorTests/Source/EditorTests/Private/UnrealEd/EditorIterationTests.cpp
+++ b/Engine/Plugins/Tests/EditorTests/Source/EditorTests/Private/UnrealEd/EditorIterationTests.cpp
@@ -60,18 +60,22 @@ bool FIterationOpenAssets::RunTest(const FString& LongAssetPath)
 	AddCommand(new FCloseAllAssetEditorsCommand());
 	AddCommand(new FFunctionLatentCommand([LongAssetPath] {
 		TRACE_BOOKMARK(TEXT("LoadAssetAndPIE - %s"), *LongAssetPath);
+		TRACE_BOOKMARK(TEXT("LoadAsset - %s"), *LongAssetPath);
 		return true;
 		}));
 
 	// Issue Load request
 	AddCommand(new FOpenEditorForAssetCommand(LongAssetPath));
-	AddCommand(new FWaitLatentCommand(0.5f));
 
 	// Wait on all async asset processing
 	AddCommand(new FFunctionLatentCommand([] {
 		return FAssetCompilingManager::Get().GetNumRemainingAssets() == 0;
 	}));
-	AddCommand(new FWaitLatentCommand(0.5f));
+	AddCommand(new FFunctionLatentCommand([LongAssetPath] {
+		TRACE_BOOKMARK(TEXT("LoadAssetComplete - %s"), *LongAssetPath);
+		TRACE_BOOKMARK(TEXT("PIE - %s"), *LongAssetPath);
+		return true;
+		}));
 
 	// Do many frames of PIE (not a time span since PIE can take a while to init)
 	AddCommand(new FDelayedFunctionLatentCommand([] {
@@ -82,10 +86,10 @@ bool FIterationOpenAssets::RunTest(const FString& LongAssetPath)
 		return GFrameCounter > (FrameNumber + 600);
 		}));
 	AddCommand(new FEndPlayMapCommand());
-	AddCommand(new FWaitLatentCommand(0.5f));
 
 	// Teardown
 	AddCommand(new FFunctionLatentCommand([LongAssetPath] {
+		TRACE_BOOKMARK(TEXT("PIEComplete - %s"), *LongAssetPath);
 		TRACE_BOOKMARK(TEXT("LoadAssetAndPIEComplete - %s"), *LongAssetPath);
 		return true;
 		}));
