EditorPerf

* Remove LoadAssetAndPIE - we have each separately now
* Simplify/organize the telemetry data a bit better: plot stall count together with frame count, move all FEngineLoop::Tick Duration stats together, move boot up stats into "Events" alongside Load and PIE events

Skipping preflight since this is just telemetry configuration and deleting two lines of code, and just in the automated test modules.

#rnx

diff --git a/Engine/Plugins/Tests/EditorTests/Source/EditorTests/Private/UnrealEd/EditorIterationTests.cpp b/Engine/Plugins/Tests/EditorTests/Source/EditorTests/Private/UnrealEd/EditorIterationTests.cpp
--- a/Engine/Plugins/Tests/EditorTests/Source/EditorTests/Private/UnrealEd/EditorIterationTests.cpp
+++ b/Engine/Plugins/Tests/EditorTests/Source/EditorTests/Private/UnrealEd/EditorIterationTests.cpp
@@ -59,7 +59,6 @@ bool FIterationOpenAssets::RunTest(const FString& LongAssetPath)
 	// Setup
 	AddCommand(new FCloseAllAssetEditorsCommand());
 	AddCommand(new FFunctionLatentCommand([LongAssetPath] {
-		TRACE_BOOKMARK(TEXT("LoadAssetAndPIE - %s"), *LongAssetPath);
 		TRACE_BOOKMARK(TEXT("LoadAsset - %s"), *LongAssetPath);
 		return true;
 		}));
@@ -90,7 +89,6 @@ bool FIterationOpenAssets::RunTest(const FString& LongAssetPath)
 	// Teardown
 	AddCommand(new FFunctionLatentCommand([LongAssetPath] {
 		TRACE_BOOKMARK(TEXT("PIEComplete - %s"), *LongAssetPath);
-		TRACE_BOOKMARK(TEXT("LoadAssetAndPIEComplete - %s"), *LongAssetPath);
 		return true;
 		}));
 
