Make PlayMapInPie tests optional

#rnx

diff --git a/Engine/Plugins/Tests/EditorTests/Source/EditorTests/Private/UnrealEd/PlayMapInPIETests.cpp b/Engine/Plugins/Tests/EditorTests/Source/EditorTests/Private/UnrealEd/PlayMapInPIETests.cpp
--- a/Engine/Plugins/Tests/EditorTests/Source/EditorTests/Private/UnrealEd/PlayMapInPIETests.cpp
+++ b/Engine/Plugins/Tests/EditorTests/Source/EditorTests/Private/UnrealEd/PlayMapInPIETests.cpp
@@ -77,6 +77,12 @@ public:
 	 */
 	virtual void GetTests(TArray<FString>& OutBeautifiedNames, TArray<FString>& OutTestCommands) const override
 	{
+		const UAutomationTestSettings* AutomationTestSettings = GetDefault<UAutomationTestSettings>();
+		if (!AutomationTestSettings->bUseAllProjectMapsToPlayInPIE)
+		{
+			return;
+		}
+
 		IAssetRegistry& AssetRegistry = FModuleManager::Get().LoadModuleChecked<FAssetRegistryModule>("AssetRegistry").Get();
 
 		if (!AssetRegistry.IsLoadingAssets())
diff --git a/Engine/Source/Runtime/Engine/Classes/Tests/AutomationTestSettings.h b/Engine/Source/Runtime/Engine/Classes/Tests/AutomationTestSettings.h
--- a/Engine/Source/Runtime/Engine/Classes/Tests/AutomationTestSettings.h
+++ b/Engine/Source/Runtime/Engine/Classes/Tests/AutomationTestSettings.h
@@ -379,6 +379,12 @@ public:
 	UPROPERTY(EditAnywhere, config, Category = "PIE Test Maps")
 	TArray<FString> MapsToPIETest;
 
+	/**
+	 * Use all Maps from project for PlayMapInPIE test
+	 */
+	UPROPERTY(EditAnywhere, config, Category = "Play all project Maps In PIE")
+	bool bUseAllProjectMapsToPlayInPIE;
+
 	/**
 	* Editor build promotion test settings
 	*/
