Rename PlayMapInPie Test to Project.Maps.AllInPIE

#rnx

diff --git a/Engine/Plugins/Tests/EditorTests/Source/EditorTests/Private/UnrealEd/PlayMapInPIETests.cpp b/Engine/Plugins/Tests/EditorTests/Source/EditorTests/Private/UnrealEd/PlayMapInPIETests.cpp
--- a/Engine/Plugins/Tests/EditorTests/Source/EditorTests/Private/UnrealEd/PlayMapInPIETests.cpp
+++ b/Engine/Plugins/Tests/EditorTests/Source/EditorTests/Private/UnrealEd/PlayMapInPIETests.cpp
@@ -190,7 +190,7 @@ protected:
 };
 
 // Editor only tests
-IMPLEMENT_CUSTOM_COMPLEX_AUTOMATION_TEST(FPlayMapInPIE, FPlayMapInPIEBase, "Content.PlayMapInPIE", (EAutomationTestFlags::EditorContext | EAutomationTestFlags::ProductFilter))
+IMPLEMENT_CUSTOM_COMPLEX_AUTOMATION_TEST(FPlayMapInPIE, FPlayMapInPIEBase, "Project.Maps.AllInPIE", (EAutomationTestFlags::EditorContext | EAutomationTestFlags::ProductFilter))
 
 void FPlayMapInPIE::GetTests(TArray<FString>& OutBeautifiedNames, TArray<FString>& OutTestCommands) const
 {
