Update AutomationController Log settings

Introduce bElevateLogWarningsToErrors as a good name for existing functionality. Note it's assuming the name as it's known by its customer code. This change lines up the settings variable to its colloquial name around the codebase.

Deprecate bTreatLogWarningsAsTestErrors as it was misleading because it doesn't really shuffle log warnings into the test errors bucket as much as it just reclassifies the log verbosity from log warning to log error.

Remove bTreatLogErrorsAsTestErrors as it wasn't functional outside of EditorIterationTest's custom base class, and it's tempting to use incorrectly in place of bSuppressLogErrors

#rnx

diff --git a/Engine/Build/EditorPerf.xml b/Engine/Build/EditorPerf.xml
--- a/Engine/Build/EditorPerf.xml
+++ b/Engine/Build/EditorPerf.xml
@@ -71,7 +71,7 @@
 	<!-- Gauntlet arguments for profile runs -->
 	<Property Name="ProfileArgs" Value="-test=EditorTests -testfilter=Group:Iteration -tracefiletrunc"/>
 	<Property Name="ProfileArgs" Value="$(ProfileArgs) -EnablePlugins=$(ExtraPluginsCommaDelims)" If="'$(ExtraPluginsCommaDelims)' != ''"/>
-	<Property Name="ProfileArgs" Value="$(ProfileArgs) -TreatLogWarningsAsTestErrors=false -TreatLogErrorsAsTestErrors=false"/>
+	<Property Name="ProfileArgs" Value="$(ProfileArgs) -SuppressLogWarnings -SuppressLogErrors"/>
 	<Property Name="ProfileArgs" Value="$(ProfileArgs) -GameInstanceLostTimerSeconds=3600"/>
 	<Property Name="ProfileArgs" Value="$(ProfileArgs) -NoLoadLevelAtStartup -Benchmarking -NODDCCLEANUP"/>
 	<Property Name="ProfileArgs" Value="$(ProfileArgs) -NoShaderDistrib -VerboseShaderLogging"/>
diff --git a/Engine/Plugins/Tests/EditorTests/Source/EditorTests/Private/UnrealEd/EditorIterationTests.cpp b/Engine/Plugins/Tests/EditorTests/Source/EditorTests/Private/UnrealEd/EditorIterationTests.cpp
--- a/Engine/Plugins/Tests/EditorTests/Source/EditorTests/Private/UnrealEd/EditorIterationTests.cpp
+++ b/Engine/Plugins/Tests/EditorTests/Source/EditorTests/Private/UnrealEd/EditorIterationTests.cpp
@@ -15,32 +15,11 @@
 #include "Subsystems/AssetEditorSubsystem.h"
 #include "Editor.h"
 
-class FFIterationOpenAssetsBase : public FAutomationTestBase
-{
-public:
-	FFIterationOpenAssetsBase(const FString& InName, const bool bInComplexTask)
-		: FAutomationTestBase(InName, bInComplexTask)
-	{
-		const UAutomationControllerSettings* Settings = GetDefault<UAutomationControllerSettings>();
-		bSuppressLogErrors = Settings->bSuppressLogErrors || !Settings->bTreatLogErrorsAsTestErrors;
-		bElevateLogWarningsToErrors = Settings->bTreatLogWarningsAsTestErrors;
-	}
-
-	virtual bool SuppressLogErrors() override { return bSuppressLogErrors; }
-
-	virtual bool ElevateLogWarningsToErrors() override { return bElevateLogWarningsToErrors; }
-
-private:
-	bool bSuppressLogErrors = false;
-	bool bElevateLogWarningsToErrors = true;
-
-};
-
 /**
  * Test to open the sub editor windows for a specified list of assets.
  * This list can be setup in the Editor Preferences window within the editor or the DefaultEngine.ini file for that particular project.
 */
-IMPLEMENT_CUSTOM_COMPLEX_AUTOMATION_TEST(FIterationOpenAssets, FFIterationOpenAssetsBase, "Project.Iteration.OpenAssets", EAutomationTestFlags::EditorContext | EAutomationTestFlags::ProductFilter);
+IMPLEMENT_CUSTOM_COMPLEX_AUTOMATION_TEST(FIterationOpenAssets, FAutomationTestBase, "Project.Iteration.OpenAssets", EAutomationTestFlags::EditorContext | EAutomationTestFlags::ProductFilter);
 
 void FIterationOpenAssets::GetTests(TArray<FString>& OutBeautifiedNames, TArray<FString>& OutTestCommands) const
 {
diff --git a/Engine/Source/Developer/AutomationController/Private/AutomationControllerSettings.cpp b/Engine/Source/Developer/AutomationController/Private/AutomationControllerSettings.cpp
--- a/Engine/Source/Developer/AutomationController/Private/AutomationControllerSettings.cpp
+++ b/Engine/Source/Developer/AutomationController/Private/AutomationControllerSettings.cpp
@@ -3,14 +3,26 @@
 
 #include "AutomationControllerSettings.h"
 
+DEFINE_LOG_CATEGORY_STATIC(LogAutomationControllerSettings, Log, All)
 
 UAutomationControllerSettings::UAutomationControllerSettings(const FObjectInitializer& ObjectInitializer)
 	: Super(ObjectInitializer)
 	, bSuppressLogErrors(false)
 	, bSuppressLogWarnings(false)
-	, bTreatLogErrorsAsTestErrors(true)
+	, bElevateLogWarningsToErrors(true)
 	, bTreatLogWarningsAsTestErrors(true)
 	, GameInstanceLostTimerSeconds(300.0f)
 	, bResetTelemetryStorageOnNewSession(false)
 {
 }
+
+void UAutomationControllerSettings::PostInitProperties()
+{
+	Super::PostInitProperties();
+
+	if (!bTreatLogWarningsAsTestErrors)
+	{
+		UE_LOG(LogAutomationControllerSettings, Warning, TEXT("UAutomationControllerSettings::bTreatLogWarningsAsTestErrors is deprecated. Use bElevateLogWarningsToErrors instead."));
+		bElevateLogWarningsToErrors = false;
+	}
+}
\ No newline at end of file
diff --git a/Engine/Source/Developer/AutomationController/Public/AutomationControllerSettings.h b/Engine/Source/Developer/AutomationController/Public/AutomationControllerSettings.h
--- a/Engine/Source/Developer/AutomationController/Public/AutomationControllerSettings.h
+++ b/Engine/Source/Developer/AutomationController/Public/AutomationControllerSettings.h
@@ -79,14 +79,16 @@ public:
 	UPROPERTY(Config)
 	bool bSuppressLogWarnings;
 
-	/** Whether to treat log errors as test errors (default=true) */
+	/** Whether to treat log warnings as log errors (default=true) */
 	UPROPERTY(Config)
-	bool bTreatLogErrorsAsTestErrors;
-
+	bool bElevateLogWarningsToErrors;
+	
+private:
 	/** Whether to treat log warnings as test errors (default=true) */
-	UPROPERTY(Config)
+	UPROPERTY(Config, Meta = (DeprecatedProperty, DeprecationMessage = "Use bElevateLogWarningsToErrors instead."))
 	bool bTreatLogWarningsAsTestErrors;
 	
+public:
 	/** How long to wait between test updates (default=1sec)*/
 	UPROPERTY(Config)
 	float CheckTestIntervalSeconds;
@@ -102,4 +104,6 @@ public:
 	/** Whether to reset data stored in telemetry file (default=false) */
 	UPROPERTY(Config)
 	bool bResetTelemetryStorageOnNewSession;
+
+	virtual void PostInitProperties() override;
 };
diff --git a/Engine/Source/Developer/FunctionalTesting/Private/FunctionalTestBase.cpp b/Engine/Source/Developer/FunctionalTesting/Private/FunctionalTestBase.cpp
--- a/Engine/Source/Developer/FunctionalTesting/Private/FunctionalTestBase.cpp
+++ b/Engine/Source/Developer/FunctionalTesting/Private/FunctionalTestBase.cpp
@@ -27,7 +27,7 @@ void FFunctionalTestBase::SetLogErrorAndWarningHandlingToDefault()
 	UAutomationControllerSettings* Settings = UAutomationControllerSettings::StaticClass()->GetDefaultObject<UAutomationControllerSettings>();
 	bSuppressLogErrors = Settings->bSuppressLogErrors;
 	bSuppressLogWarnings = Settings->bSuppressLogWarnings;
-	bElevateLogWarningsToErrors = Settings->bTreatLogWarningsAsTestErrors;
+	bElevateLogWarningsToErrors = Settings->bElevateLogWarningsToErrors;
 }
 
 /**
diff --git a/Engine/Source/Programs/AutomationTool/Gauntlet/Unreal/Editor/EditorTests.cs b/Engine/Source/Programs/AutomationTool/Gauntlet/Unreal/Editor/EditorTests.cs
--- a/Engine/Source/Programs/AutomationTool/Gauntlet/Unreal/Editor/EditorTests.cs
+++ b/Engine/Source/Programs/AutomationTool/Gauntlet/Unreal/Editor/EditorTests.cs
@@ -41,13 +41,13 @@ namespace UnrealEditor
 		/// Control for interpretation of log warnings as test failures
 		/// </summary>
 		[AutoParam]
-		public string TreatLogWarningsAsTestErrors = "true";
+		public bool SuppressLogWarnings = false;
 
 		/// <summary>
 		/// Control for interpretation of log errors as test failures
 		/// </summary>
 		[AutoParam]
-		public string TreatLogErrorsAsTestErrors = "true";
+		public bool SuppressLogErrors = false;
 
 		/// <summary>
 		/// Modify the game instance lost timeout interval
@@ -99,14 +99,14 @@ namespace UnrealEditor
 				AppConfig.CommandLineParams.Add("trace=cpu,counters,bookmarks");
 			}
 
-			if (TreatLogWarningsAsTestErrors.ToLower() == "false" || TreatLogWarningsAsTestErrors == "0")
+			if (SuppressLogWarnings)
 			{
-				AppConfig.CommandLineParams.Add("ini:Engine:[/Script/AutomationController.AutomationControllerSettings]:bTreatLogWarningsAsTestErrors=false");
+				AppConfig.CommandLineParams.Add("ini:Engine:[/Script/AutomationController.AutomationControllerSettings]:bSuppressLogWarnings=true");
 			}
 
-			if (TreatLogErrorsAsTestErrors.ToLower() == "false" || TreatLogErrorsAsTestErrors == "0")
+			if (SuppressLogErrors)
 			{
-				AppConfig.CommandLineParams.Add("ini:Engine:[/Script/AutomationController.AutomationControllerSettings]:bTreatLogErrorsAsTestErrors=false");
+				AppConfig.CommandLineParams.Add("ini:Engine:[/Script/AutomationController.AutomationControllerSettings]:bSuppressLogErrors=true");
 			}
 
 			if (GameInstanceLostTimerSeconds != string.Empty)
diff --git a/Engine/Source/Runtime/Core/Private/Misc/AutomationTest.cpp b/Engine/Source/Runtime/Core/Private/Misc/AutomationTest.cpp
--- a/Engine/Source/Runtime/Core/Private/Misc/AutomationTest.cpp
+++ b/Engine/Source/Runtime/Core/Private/Misc/AutomationTest.cpp
@@ -31,7 +31,7 @@ CORE_API ELogVerbosity::Type GetAutomationLogLevel(ELogVerbosity::Type LogVerbos
 	// can leverage, not just functional tests
 	static bool bSuppressLogWarnings = false;
 	static bool bSuppressLogErrors = false;
-	static bool bTreatLogWarningsAsTestErrors = false;
+	static bool bElevateLogWarningsToErrors = false;
 
 	static FAutomationTestBase* LastTest = nullptr;
 
@@ -40,7 +40,7 @@ CORE_API ELogVerbosity::Type GetAutomationLogLevel(ELogVerbosity::Type LogVerbos
 		// These can be changed in the editor so can't just be cached for the whole session
 		GConfig->GetBool(TEXT("/Script/AutomationController.AutomationControllerSettings"), TEXT("bSuppressLogErrors"), bSuppressLogErrors, GEngineIni);
 		GConfig->GetBool(TEXT("/Script/AutomationController.AutomationControllerSettings"), TEXT("bSuppressLogWarnings"), bSuppressLogWarnings, GEngineIni);
-		GConfig->GetBool(TEXT("/Script/AutomationController.AutomationControllerSettings"), TEXT("bTreatLogWarningsAsTestErrors"), bTreatLogWarningsAsTestErrors, GEngineIni);
+		GConfig->GetBool(TEXT("/Script/AutomationController.AutomationControllerSettings"), TEXT("bElevateLogWarningsToErrors"), bElevateLogWarningsToErrors, GEngineIni);
 		LastTest = CurrentTest;
 	}
 
@@ -58,7 +58,7 @@ CORE_API ELogVerbosity::Type GetAutomationLogLevel(ELogVerbosity::Type LogVerbos
 				{
 					EffectiveVerbosity = ELogVerbosity::NoLogging;
 				}
-				else if (CurrentTest->ElevateLogWarningsToErrors() || bTreatLogWarningsAsTestErrors)
+				else if (CurrentTest->ElevateLogWarningsToErrors() || bElevateLogWarningsToErrors)
 				{
 					EffectiveVerbosity = ELogVerbosity::Error;
 				}
