Fix GatherObjectReferencersForDeletion that could prevent deletion in the case of external references with multi-level indirection that can actually be GCed
Add basic unittest for GatherObjectReferencersForDeletion to validate the fix
  - Slight regression in performance due to more processing required to walk up the referer chain
  - 1.2s to 8.8s to press delete on all (In use by level) assets in QAGame / DM-Deck
  - Editor.UseLegacyGetReferencersForDeletion = 1 takes 1m22s for the same result

diff --git a/Engine/Config/BaseEngine.ini b/Engine/Config/BaseEngine.ini
--- a/Engine/Config/BaseEngine.ini
+++ b/Engine/Config/BaseEngine.ini
@@ -1866,8 +1866,6 @@ r.VSync=0
 ; Parallel rendering has not been tested in the editor so just disable it.
 r.RHICmdBypass=0
 
-Editor.UseLegacyGetReferencersForDeletion=1
-
 [SystemSettingsSplitScreen2]
 ; System settings overrides for split screen
 ; Use medium detail mode in splitscreen, this allows LD's to mark meshes as High detail and they won't render in SS
diff --git a/Engine/Plugins/Tests/EditorTests/Source/EditorTests/Private/UnrealEd/ObjectToolsTests.cpp b/Engine/Plugins/Tests/EditorTests/Source/EditorTests/Private/UnrealEd/ObjectToolsTests.cpp
new file mode 100644
--- /dev/null
+++ b/Engine/Plugins/Tests/EditorTests/Source/EditorTests/Private/UnrealEd/ObjectToolsTests.cpp
@@ -0,0 +1,262 @@
+// Copyright Epic Games, Inc. All Rights Reserved.
+
+#include "ObjectToolsTests.h"
+
+#include "Misc/Paths.h"
+#include "Misc/AutomationTest.h"
+#include "ScopedTransaction.h"
+#include "Tests/AutomationTestSettings.h"
+#include "Tests/AutomationEditorCommon.h"
+#include "Tests/AutomationCommon.h"
+
+class FObjectToolsTests_GatherObjectReferencersForDeletionTestBase : public FAutomationTestBase
+{
+public:
+	FObjectToolsTests_GatherObjectReferencersForDeletionTestBase(const FString& InName, const bool bInComplexTask)
+		: FAutomationTestBase(InName, bInComplexTask)
+	{
+	}
+
+protected:
+	void ResetUndoBuffer()
+	{
+		GEditor->Trans->Reset(FText::FromString("Discard undo history during ObjectTools testing."));
+	}
+
+	void TestGatherObjectReferencersForDeletion(UObject* Object, const TCHAR* MessageForReference, bool bExpectedIsReferencedValue, const TCHAR* MessageForTransactor, bool bExpectedIsReferencedByUndoValue)
+	{
+		bool bIsReferenced = false;
+		bool bIsReferencedByUndo = false;
+		ObjectTools::GatherObjectReferencersForDeletion(Object, bIsReferenced, bIsReferencedByUndo);
+		TestEqual(MessageForReference, bIsReferenced, bExpectedIsReferencedValue);
+		TestEqual(MessageForTransactor, bIsReferencedByUndo, bExpectedIsReferencedByUndoValue);
+	}
+
+	void TestGatherObjectReferencersForDeletion(UObject* Object, const TCHAR* Message, bool bExpectedIsReferencedValue, bool bExpectedIsReferencedByUndoValue)
+	{
+		TestGatherObjectReferencersForDeletion(Object, Message, bExpectedIsReferencedValue, Message, bExpectedIsReferencedByUndoValue);
+	}
+
+	enum class EReferenceType
+	{
+		Weak,
+		Strong
+	};
+
+	enum class EReferenceChainType
+	{
+		Direct,
+		Indirect
+	};
+
+	const TCHAR* LexToString(EReferenceType InReferenceType)
+	{
+		switch (InReferenceType)
+		{
+		case EReferenceType::Weak:
+			return TEXT("Weak");
+		case EReferenceType::Strong:
+			return TEXT("Strong");
+		default:
+			checkNoEntry();
+			return TEXT("Unknown");
+		}
+	}
+
+	const TCHAR* LexToString(EReferenceChainType InReferenceChainType)
+	{
+		switch (InReferenceChainType)
+		{
+		case EReferenceChainType::Direct:
+			return TEXT("Direct");
+		case EReferenceChainType::Indirect:
+			return TEXT("Indirect");
+		default:
+			checkNoEntry();
+			return TEXT("Unknown");
+		}
+	}
+
+	void SetObjectToolsTestObjectReference(UObjectToolsTestObject* InReferencer, UObject* InReferencee, EReferenceType InRerenceType)
+	{
+		switch (InRerenceType)
+		{
+			case EReferenceType::Weak:
+				InReferencer->WeakReference = InReferencee;
+				break;
+			case EReferenceType::Strong:
+				InReferencer->StrongReference = InReferencee;
+				break;
+			default:
+				checkNoEntry();
+				break;
+		}
+	}
+
+	bool GetExpectedResultFromReferenceType(EReferenceType InRerenceType)
+	{
+		switch (InRerenceType)
+		{
+		case EReferenceType::Weak:
+			return false;	// Weak reference are expected to not return a valid referencer
+		case EReferenceType::Strong:
+			return true;    // Strong reference are expected to return a valid referencer
+		default:
+			checkNoEntry();
+			return false;
+		}
+	}
+
+	bool GetExpectedTransactorResultFromReferenceType(EReferenceType InRerenceType)
+	{
+		switch (InRerenceType)
+		{
+		case EReferenceType::Weak:
+			return true;	// The current behavior of the transactor is to convert weak reference to strong ones while they are in the undo buffer. This can be made false when the transactor behavior is fixed in the future.
+		case EReferenceType::Strong:
+			return true;    // Strong reference are expected to return a valid referencer
+		default:
+			checkNoEntry();
+			return false;
+		}
+	}
+
+	// Same boilerplate code that can be used with different types of references
+	void TestWithTransactor(UPackage* TempPackage, UObjectToolsTestObject* InReferencer, UObject* InReferencee, EReferenceChainType InReferenceChainType, EReferenceType InReferenceType)
+	{
+		const bool bExpectedResult = GetExpectedResultFromReferenceType(InReferenceType);
+		const TCHAR* ExpectedAction = bExpectedResult ? TEXT("detect") : TEXT("ignore");
+		const bool bTransactorExpectedResult = GetExpectedTransactorResultFromReferenceType(InReferenceType);
+		const TCHAR* ExpectedTransactorAction = bTransactorExpectedResult ? TEXT("detect") : TEXT("ignore");
+		const FString InReferenceTypeName = FString(LexToString(InReferenceType)).ToLower();
+		const FString InReferenceChainTypeName = FString(LexToString(InReferenceChainType)).ToLower();
+		
+		// Not strictly necessary, just a precaution in case some other test forgot to clear the undo buffer
+		ResetUndoBuffer();
+
+		// The property has not been set yet, we expect no referencer to be returned
+		TestGatherObjectReferencersForDeletion(TempPackage, TEXT("GatherObjectReferencersForDeletion shouldn't detect any reference on any objects inside TempPackage"), false, false);
+
+		// Set our reference
+		SetObjectToolsTestObjectReference(InReferencer, InReferencee, InReferenceType);
+		
+		// Test the basic case with our reference now set
+		TestGatherObjectReferencersForDeletion(TempPackage, *FString::Printf(TEXT("GatherObjectReferencersForDeletion should %s %s %s references on objects inside TempPackage"), ExpectedAction, *InReferenceChainTypeName, *InReferenceTypeName), bExpectedResult, false);
+
+		// Set the reference to null in a transaction so that the transactor now holds onto it
+		{
+			FScopedTransaction Transaction(FText::FromString("ObjectTools testing."));
+			InReferencer->PreEditChange(nullptr);
+			SetObjectToolsTestObjectReference(InReferencer, nullptr, InReferenceType);
+			InReferencer->PostEditChange();
+		}
+
+		// Validate that the transactor works as expected otherwise the test would be broken
+		check(GEditor->Trans->IsObjectInTransationBuffer(InReferencer));
+
+		// The reference has been set to nullptr, only the transactor should now have a reference
+		TestGatherObjectReferencersForDeletion(
+			TempPackage, 
+			*FString::Printf(TEXT("GatherObjectReferencersForDeletion shouldn't detect non-transactor %s %s references at this point"), *InReferenceChainTypeName, *InReferenceTypeName), false, 
+			*FString::Printf(TEXT("GatherObjectReferencersForDeletion should %s %s %s references from the transactor"), ExpectedTransactorAction , *InReferenceChainTypeName, *InReferenceTypeName), bTransactorExpectedResult
+		);
+
+		// Restore the reference so that both the transactor and normal object have valid references
+		SetObjectToolsTestObjectReference(InReferencer, InReferencee, InReferenceType);
+
+		// Both the transactor and the object should now have a reference reported
+		TestGatherObjectReferencersForDeletion(TempPackage, 
+			*FString::Printf(TEXT("GatherObjectReferencersForDeletion should %s %s %s references on any objects inside TempPackage"), ExpectedAction, *InReferenceChainTypeName, *InReferenceTypeName), bExpectedResult, 
+			*FString::Printf(TEXT("GatherObjectReferencersForDeletion should %s %s %s references from the transactor"), ExpectedTransactorAction, *InReferenceChainTypeName, *InReferenceTypeName), bTransactorExpectedResult);
+
+		// This should clear the undo reference
+		ResetUndoBuffer();
+
+		// Test that the transactor is not considered a referencer anymore
+		TestGatherObjectReferencersForDeletion(TempPackage, *FString::Printf(TEXT("GatherObjectReferencersForDeletion should %s %s %s references on objects inside TempPackage"), ExpectedAction, *InReferenceChainTypeName, *InReferenceTypeName), bExpectedResult, false);
+
+		// Set the reference to null so that we exit the test in the same state we came in
+		SetObjectToolsTestObjectReference(InReferencer, nullptr, InReferenceType);
+
+		// Last check to make sure there is no more referencer left
+		TestGatherObjectReferencersForDeletion(TempPackage, TEXT("GatherObjectReferencersForDeletion shouldn't detect any reference on any objects inside TempPackage"), false, false);
+	}
+};
+
+/**
+* Automation test for validating GatherObjectReferencersForDeletion functionality
+*/
+IMPLEMENT_CUSTOM_SIMPLE_AUTOMATION_TEST(
+	FObjectToolsTests_GatherObjectReferencersForDeletion, 
+	FObjectToolsTests_GatherObjectReferencersForDeletionTestBase, 
+	"Editor.ObjectTools.GatherObjectReferencersForDeletion", 
+	EAutomationTestFlags::EditorContext | EAutomationTestFlags::EngineFilter);
+
+bool FObjectToolsTests_GatherObjectReferencersForDeletion::RunTest(const FString& )
+{
+	UPackage* TempPackage = NewObject<UPackage>(nullptr, TEXT("/Engine/Editor/ObjectToolsTests/Transient"), RF_Transient);
+
+	// Make sure the MetaData of the package exists because it is standalone and can interfere with the whole reference gathering
+	UMetaData* MetaData = TempPackage->GetMetaData();
+	check(MetaData);
+
+	UObjectToolsTestObject* ObjectInPackage = NewObject<UObjectToolsTestObject>(TempPackage, NAME_None, RF_Transactional);
+
+	// Test direct references to ObjectInPackage
+	{
+		UObjectToolsTestObject* DirectRootReferencer = NewObject<UObjectToolsTestObject>(GetTransientPackage(), NAME_None, RF_Transactional);
+		DirectRootReferencer->AddToRoot();
+
+		TestWithTransactor(TempPackage, DirectRootReferencer, ObjectInPackage, EReferenceChainType::Direct, EReferenceType::Weak);
+		TestWithTransactor(TempPackage, DirectRootReferencer, ObjectInPackage, EReferenceChainType::Direct, EReferenceType::Strong);
+
+		DirectRootReferencer->RemoveFromRoot();
+	}
+
+	// Test indirect references to ObjectInPackage
+	{
+		UObjectToolsTestObject* Referencer = NewObject<UObjectToolsTestObject>(GetTransientPackage(), NAME_None, RF_Transactional);
+
+		// This object itself is not rooted, so it's reference shouldn't matter
+		Referencer->StrongReference = ObjectInPackage;
+
+		UObjectToolsTestObject* IndirectRootReferencer = NewObject<UObjectToolsTestObject>(GetTransientPackage(), NAME_None, RF_Transactional);
+		IndirectRootReferencer->AddToRoot();
+
+		TestWithTransactor(TempPackage, IndirectRootReferencer, Referencer, EReferenceChainType::Indirect, EReferenceType::Weak);
+		TestWithTransactor(TempPackage, IndirectRootReferencer, Referencer, EReferenceChainType::Indirect, EReferenceType::Strong);
+
+		IndirectRootReferencer->RemoveFromRoot();
+	}
+
+	// Test when the transaction buffer has a direct reference to one of the objects we're trying to destroy
+	{
+		UObjectToolsTestObject* DummyObject = NewObject<UObjectToolsTestObject>(GetTransientPackage(), NAME_None, RF_Transactional);
+
+		TestGatherObjectReferencersForDeletion(TempPackage, TEXT("GatherObjectReferencersForDeletion shouldn't detect any reference on any objects inside TempPackage"), false, false);
+
+		{
+			FScopedTransaction Transaction(FText::FromString("ObjectTools testing."));
+			ObjectInPackage->PreEditChange(nullptr);
+			ObjectInPackage->StrongReference = DummyObject;
+			ObjectInPackage->PostEditChange();
+		}
+
+		TestGatherObjectReferencersForDeletion(
+			TempPackage,
+			TEXT("GatherObjectReferencersForDeletion shouldn't detect non-transactor references at this point"), false,
+			TEXT("GatherObjectReferencersForDeletion should detect a transactor reference"), true
+		);
+
+		// Restore the property as it was before the test in case we add more test in the future
+		ObjectInPackage->StrongReference = nullptr;
+
+		// This should clear the undo reference
+		ResetUndoBuffer();
+
+		// Make sure we're back to initial state
+		TestGatherObjectReferencersForDeletion(TempPackage, TEXT("GatherObjectReferencersForDeletion shouldn't detect any reference on any objects inside TempPackage"), false, false);
+	}
+
+	return true;
+}
\ No newline at end of file
diff --git a/Engine/Plugins/Tests/EditorTests/Source/EditorTests/Private/UnrealEd/ObjectToolsTests.h b/Engine/Plugins/Tests/EditorTests/Source/EditorTests/Private/UnrealEd/ObjectToolsTests.h
new file mode 100644
--- /dev/null
+++ b/Engine/Plugins/Tests/EditorTests/Source/EditorTests/Private/UnrealEd/ObjectToolsTests.h
@@ -0,0 +1,21 @@
+// Copyright Epic Games, Inc. All Rights Reserved.
+
+#pragma once
+
+#include "CoreMinimal.h"
+#include "UObject/Object.h"
+#include "UObject/ObjectMacros.h"
+#include "ObjectToolsTests.generated.h"
+
+UCLASS()
+class UObjectToolsTestObject : public UObject
+{
+	GENERATED_BODY()
+
+public:
+	UPROPERTY()
+	UObject* StrongReference;
+
+	UPROPERTY()
+	TWeakObjectPtr<UObject> WeakReference;
+};
diff --git a/Engine/Source/Editor/UnrealEd/Private/ObjectTools.cpp b/Engine/Source/Editor/UnrealEd/Private/ObjectTools.cpp
--- a/Engine/Source/Editor/UnrealEd/Private/ObjectTools.cpp
+++ b/Engine/Source/Editor/UnrealEd/Private/ObjectTools.cpp
@@ -180,9 +180,62 @@ namespace ObjectTools
 
 		return bIsSupported;
 	}
-	
+
+	bool IsNonGCObject(UObject* Object)
+	{
+		FUObjectItem* ObjectItem = GUObjectArray.ObjectToObjectItem(Object);
+		return
+			ObjectItem->IsRootSet() ||
+			ObjectItem->HasAnyFlags(EInternalObjectFlags::GarbageCollectionKeepFlags) ||
+			(GARBAGE_COLLECTION_KEEPFLAGS != RF_NoFlags && Object->HasAnyFlags(GARBAGE_COLLECTION_KEEPFLAGS));
+	}
+
+	TSet<UObject*> FindObjectsRoots(TSet<UObject*>& InObjects)
+	{
+		TRACE_CPUPROFILER_EVENT_SCOPE(FindObjectsRoots)
+
+		TSet<UObject*> Roots;
+		UTransactor* Transactor = GEditor ? ToRawPtr(GEditor->Trans) : nullptr;
+
+		// Handle the objects themselves if they can't be GCed
+		for (UObject* Object : InObjects)
+		{
+			if (IsNonGCObject(Object))
+			{
+				Roots.Add(Object);
+			}
+		}
+
+		// We recursively grow the cluster of objects we need to find referencers on until no more referencers are found
+		int32 LastObjectCount = 0;
+		while (InObjects.Num() != LastObjectCount)
+		{
+			LastObjectCount = InObjects.Num();
+			for (UObject* NewReferencer : FReferencerFinder::GetAllReferencers(InObjects, &InObjects))
+			{
+				// Stop walking the dependency chain when the transactor is the referencer
+				if (Transactor == NewReferencer)
+				{
+					Roots.Add(Transactor);
+				}
+				else if (IsNonGCObject(NewReferencer))
+				{
+					Roots.Add(NewReferencer);
+				}
+				else
+				{
+					InObjects.Add(NewReferencer);
+				}
+			}
+		}
+
+		return MoveTemp(Roots);
+	}
+
 	void GatherObjectReferencersForDeletion(UObject* InObject, bool& bOutIsReferenced, bool& bOutIsReferencedInMemoryByUndo, FReferencerInformationList* OutMemoryReferences, bool bInRequireReferencingProperties)
 	{
+		TRACE_CPUPROFILER_EVENT_SCOPE(GatherObjectReferencersForDeletion)
+
 		if (OutMemoryReferences)
 		{
 			OutMemoryReferences->ExternalReferences.Reset();
@@ -275,6 +328,30 @@ namespace ObjectTools
 				}
 			}
 
+			// Walk the ref chain to make sure external refs we found are actually reachable and can't be GCed
+			if (bOutIsReferenced)
+			{
+				TSet<UObject*> Referencers;
+				Referencers.Reserve(References.ExternalReferences.Num());
+
+				for (FReferencerInformation& RefInfo : References.ExternalReferences)
+				{
+					Referencers.Add(RefInfo.Referencer);
+				}
+
+				TSet<UObject*> Roots = FindObjectsRoots(Referencers);
+
+				if (Roots.Contains(Transactor))
+				{
+					bOutIsReferencedInMemoryByUndo = true;
+					bOutIsReferenced = Roots.Num() > 1;
+				}
+				else
+				{
+					bOutIsReferenced = Roots.Num() > 0;
+				}
+			}
+
 			// For now, only IsReferenced can output which Property refers to an object and it is required
 			// when showing the graph dialog of referencers. 
 			// Only called when required and only when references are found, effect of this slower path is expected to be mostly negligible.
