Add missing include not caught by preflight to fix Non-PCH Non-Unity build
Fix disallowed native pointer not caught by preflight to fix https://horde.devtools.epicgames.com/job/61fb295b033a864b77e16e4d?step=a371

These should have been caught by preflight, they are lightweight to test on modified files only

diff --git a/Engine/Plugins/Tests/EditorTests/Source/EditorTests/Private/UnrealEd/ObjectToolsTests.cpp b/Engine/Plugins/Tests/EditorTests/Source/EditorTests/Private/UnrealEd/ObjectToolsTests.cpp
--- a/Engine/Plugins/Tests/EditorTests/Source/EditorTests/Private/UnrealEd/ObjectToolsTests.cpp
+++ b/Engine/Plugins/Tests/EditorTests/Source/EditorTests/Private/UnrealEd/ObjectToolsTests.cpp
@@ -4,12 +4,14 @@
 
 #include "Misc/Paths.h"
 #include "Misc/AutomationTest.h"
+#include "Editor.h"
 #include "ObjectTools.h"
 #include "ScopedTransaction.h"
 #include "Tests/AutomationTestSettings.h"
 #include "Tests/AutomationEditorCommon.h"
 #include "Tests/AutomationCommon.h"
 
+
 class FObjectToolsTests_GatherObjectReferencersForDeletionTestBase : public FAutomationTestBase
 {
 public:
diff --git a/Engine/Plugins/Tests/EditorTests/Source/EditorTests/Private/UnrealEd/ObjectToolsTests.h b/Engine/Plugins/Tests/EditorTests/Source/EditorTests/Private/UnrealEd/ObjectToolsTests.h
--- a/Engine/Plugins/Tests/EditorTests/Source/EditorTests/Private/UnrealEd/ObjectToolsTests.h
+++ b/Engine/Plugins/Tests/EditorTests/Source/EditorTests/Private/UnrealEd/ObjectToolsTests.h
@@ -4,6 +4,7 @@
 
 #include "CoreMinimal.h"
 #include "UObject/Object.h"
+#include "UObject/ObjectPtr.h"
 #include "UObject/ObjectMacros.h"
 #include "ObjectToolsTests.generated.h"
 
@@ -14,7 +15,7 @@ class UObjectToolsTestObject : public UObject
 
 public:
 	UPROPERTY()
-	UObject* StrongReference;
+	TObjectPtr<UObject> StrongReference;
 
 	UPROPERTY()
 	TWeakObjectPtr<UObject> WeakReference;
