NoUnity/NoPCH fixes in physics files

diff --git a/Engine/Plugins/Tests/EditorTests/Source/EditorTests/Private/UnrealEd/CollisionAutomationTests.h b/Engine/Plugins/Tests/EditorTests/Source/EditorTests/Private/UnrealEd/CollisionAutomationTests.h
--- a/Engine/Plugins/Tests/EditorTests/Source/EditorTests/Private/UnrealEd/CollisionAutomationTests.h
+++ b/Engine/Plugins/Tests/EditorTests/Source/EditorTests/Private/UnrealEd/CollisionAutomationTests.h
@@ -6,6 +6,7 @@
 #include "UObject/ObjectMacros.h"
 #include "UObject/Object.h"
 #include "Engine/EngineTypes.h"
+#include "Engine/HitResult.h"
 #include "CollisionAutomationTests.generated.h"
 
 /**
diff --git a/Engine/Source/Runtime/Engine/Private/Collision/KAggregateGeom.cpp b/Engine/Source/Runtime/Engine/Private/Collision/KAggregateGeom.cpp
--- a/Engine/Source/Runtime/Engine/Private/Collision/KAggregateGeom.cpp
+++ b/Engine/Source/Runtime/Engine/Private/Collision/KAggregateGeom.cpp
@@ -6,6 +6,7 @@
 
 #include "CoreMinimal.h"
 #include "EngineDefines.h"
+#include "EngineLogs.h"
 #include "PhysicsEngine/ShapeElem.h"
 #include "PhysicsEngine/ConvexElem.h"
 #include "PhysicsEngine/BoxElem.h"
diff --git a/Engine/Source/Runtime/Engine/Private/PhysicsEngine/CollisionStructure.h b/Engine/Source/Runtime/Engine/Private/PhysicsEngine/CollisionStructure.h
--- a/Engine/Source/Runtime/Engine/Private/PhysicsEngine/CollisionStructure.h
+++ b/Engine/Source/Runtime/Engine/Private/PhysicsEngine/CollisionStructure.h
@@ -5,6 +5,7 @@
 #include "PhysicsEngine/AggregateGeom.h"
 #include "PhysicsEngine/CollisionQueryFilterCallback.h"
 
+struct FHitResult;
 
 class FCollisionStructure
 {
