Add tests to ObjectTools test suite

#rnx

diff --git a/Engine/Plugins/Tests/EditorTests/Source/EditorTests/Private/UnrealEd/ObjectToolsTests.cpp b/Engine/Plugins/Tests/EditorTests/Source/EditorTests/Private/UnrealEd/ObjectToolsTests.cpp
--- a/Engine/Plugins/Tests/EditorTests/Source/EditorTests/Private/UnrealEd/ObjectToolsTests.cpp
+++ b/Engine/Plugins/Tests/EditorTests/Source/EditorTests/Private/UnrealEd/ObjectToolsTests.cpp
@@ -261,5 +261,42 @@ bool FObjectToolsTests_GatherObjectReferencersForDeletion::RunTest(const FString
 		TestGatherObjectReferencersForDeletion(TempPackage, TEXT("GatherObjectReferencersForDeletion shouldn't detect any reference on any objects inside TempPackage"), false, false);
 	}
 
+	return true;
+}
+
+/**
+* Automation test for validating ForceDeleteObjects functionality
+*/
+IMPLEMENT_SIMPLE_AUTOMATION_TEST(
+	FObjectToolsTests_ForceDeleteObjects,
+	"Editor.ObjectTools.ForceDeleteObjects",
+	EAutomationTestFlags::EditorContext | EAutomationTestFlags::EngineFilter);
+
+bool FObjectToolsTests_ForceDeleteObjects::RunTest(const FString&)
+{
+	// Simple test for now to repro an existing issue highlighting standalone usage on assets.
+	// If we try to delete a group of objects having reference between each other with some of them
+	// being stand-alone, the function doing the delete should ignore standalone referencers if 
+	// the referencers are also part of the group of objects being deleted.
+
+	UPackage* TempPackageA = NewObject<UPackage>(nullptr, TEXT("/Engine/Editor/ObjectToolsTests/Transient"), RF_Transient);
+	UPackage* TempPackageB = NewObject<UPackage>(nullptr, TEXT("/Engine/Editor/ObjectToolsTests/Transient"), RF_Transient);
+
+	// Make sure the MetaData of the package exists because it is standalone and can interfere with the whole reference gathering
+	verify(TempPackageA->GetMetaData());
+	verify(TempPackageB->GetMetaData());
+
+	UObjectToolsTestObject* ObjectInPackageA = NewObject<UObjectToolsTestObject>(TempPackageA, NAME_None, RF_Transactional | RF_Standalone);
+	UObjectToolsTestObject* ObjectInPackageB = NewObject<UObjectToolsTestObject>(TempPackageB, NAME_None, RF_Transactional | RF_Standalone);
+
+	ObjectInPackageA->StrongReference = ObjectInPackageB;
+	ObjectInPackageB->StrongReference = ObjectInPackageA;
+
+	TestEqual(
+		TEXT("ForceDeleteObjects should be able to force delete both objects without warnings or errors"), 
+		ObjectTools::ForceDeleteObjects({ ObjectInPackageA, ObjectInPackageB }),
+		2
+	);
+
 	return true;
 }
\ No newline at end of file
