Change FTraceAllTimelinesAutomationTest and FBlueprintRenameAndCloneTest flags to remove the RequiresUser constraint as it is not necessary. They can be run unattended or with nullrhi.

#rnx

diff --git a/Engine/Plugins/Tests/EditorTests/Source/EditorTests/Private/UnrealEd/BlueprintAutomationTests.cpp b/Engine/Plugins/Tests/EditorTests/Source/EditorTests/Private/UnrealEd/BlueprintAutomationTests.cpp
--- a/Engine/Plugins/Tests/EditorTests/Source/EditorTests/Private/UnrealEd/BlueprintAutomationTests.cpp
+++ b/Engine/Plugins/Tests/EditorTests/Source/EditorTests/Private/UnrealEd/BlueprintAutomationTests.cpp
@@ -50,7 +50,7 @@ DEFINE_LOG_CATEGORY_STATIC(LogBlueprintAutomationTests, Log, All);
 IMPLEMENT_COMPLEX_AUTOMATION_TEST(FBlueprintCompileOnLoadTest, "Project.Blueprints.Compile-On-Load", EAutomationTestFlags::EditorContext | EAutomationTestFlags::StressFilter)
 IMPLEMENT_COMPLEX_AUTOMATION_TEST(FBlueprintInstancesTest, "Project.Blueprints.Instance Test", EAutomationTestFlags::EditorContext | EAutomationTestFlags::StressFilter)
 IMPLEMENT_COMPLEX_AUTOMATION_TEST(FBlueprintReparentTest, "System.Blueprints.Reparent", EAutomationTestFlags::EditorContext | EAutomationTestFlags::StressFilter)
-IMPLEMENT_COMPLEX_AUTOMATION_TEST(FBlueprintRenameAndCloneTest, "Project.Blueprints.Rename And Clone", (EAutomationTestFlags::EditorContext | EAutomationTestFlags::RequiresUser | EAutomationTestFlags::StressFilter))
+IMPLEMENT_COMPLEX_AUTOMATION_TEST(FBlueprintRenameAndCloneTest, "Project.Blueprints.Rename And Clone", (EAutomationTestFlags::EditorContext | EAutomationTestFlags::StressFilter))
 IMPLEMENT_COMPLEX_AUTOMATION_TEST(FCompileBlueprintsTest, "Project.Blueprints.Compile Blueprints", EAutomationTestFlags::EditorContext | EAutomationTestFlags::StressFilter)
 IMPLEMENT_COMPLEX_AUTOMATION_TEST(FCompileAnimBlueprintsTest, "Project.Blueprints.Compile Anims", EAutomationTestFlags::EditorContext | EAutomationTestFlags::StressFilter)
 
diff --git a/Engine/Plugins/Tests/EditorTests/Source/EditorTests/Private/UnrealEd/EditorAutomationTests.cpp b/Engine/Plugins/Tests/EditorTests/Source/EditorTests/Private/UnrealEd/EditorAutomationTests.cpp
--- a/Engine/Plugins/Tests/EditorTests/Source/EditorTests/Private/UnrealEd/EditorAutomationTests.cpp
+++ b/Engine/Plugins/Tests/EditorTests/Source/EditorTests/Private/UnrealEd/EditorAutomationTests.cpp
@@ -716,7 +716,7 @@ bool FLightPlacement::RunTest(const FString& Parameters)
  * Unit test to find all timelines in blueprints and list the events that can trigger them.
  * Timelines implicitly tick and are usually used for cosmetic events, so they can cause performance problems on dedicated servers.
  */
-IMPLEMENT_COMPLEX_AUTOMATION_TEST(FTraceAllTimelinesAutomationTest, "Project.Performance Audits.Find Timelines On Server", (EAutomationTestFlags::EditorContext | EAutomationTestFlags::RequiresUser | EAutomationTestFlags::StressFilter))
+IMPLEMENT_COMPLEX_AUTOMATION_TEST(FTraceAllTimelinesAutomationTest, "Project.Performance Audits.Find Timelines On Server", (EAutomationTestFlags::EditorContext | EAutomationTestFlags::StressFilter))
 
 /** 
  * Requests an enumeration of all blueprints to be loaded
