Add support for Blueprint Editor Utilities in UE functional tests

diff --git a/Engine/Config/BaseEditorPerProjectUserSettings.ini b/Engine/Config/BaseEditorPerProjectUserSettings.ini
--- a/Engine/Config/BaseEditorPerProjectUserSettings.ini
+++ b/Engine/Config/BaseEditorPerProjectUserSettings.ini
@@ -805,6 +805,8 @@ DefaultBrushRadius=128
 +Node=(TargetClass=/Script/FunctionalTesting.FunctionalTest TargetEvent="ReceivePrepareTest")
 +Node=(TargetClass=/Script/FunctionalTesting.FunctionalTest TargetEvent="ReceiveStartTest")
 +Node=(TargetClass=/Script/FunctionalTesting.FunctionalTest TargetEvent="ReceiveTick")
++Node=(TargetClass=/Script/EditorTests.EditorUtilityTest TargetEvent="PrepareTest")
++Node=(TargetClass=/Script/EditorTests.EditorUtilityTest TargetEvent="StartTest")
 +Node=(TargetClass=/Script/MovieSceneCapture.UserDefinedCaptureProtocol TargetEvent="OnPreTick")
 +Node=(TargetClass=/Script/MovieSceneCapture.UserDefinedCaptureProtocol TargetEvent="OnTick")
 +Node=(TargetClass=/Script/MovieSceneCapture.UserDefinedCaptureProtocol TargetEvent="OnStartCapture")
diff --git a/Engine/Plugins/Tests/EditorTests/Source/EditorTests/Classes/EditorUtilityTest.h b/Engine/Plugins/Tests/EditorTests/Source/EditorTests/Classes/EditorUtilityTest.h
new file mode 100644
--- /dev/null
+++ b/Engine/Plugins/Tests/EditorTests/Source/EditorTests/Classes/EditorUtilityTest.h
@@ -0,0 +1,143 @@
+// Copyright Epic Games, Inc. All Rights Reserved.
+
+/**
+ * Base class of any editor-only test objects
+ */
+
+#pragma once
+
+#include "CoreMinimal.h"
+#include "UObject/ObjectMacros.h"
+#include "UObject/Object.h"
+#include "UObject/Script.h"
+#include "UObject/ScriptMacros.h"
+#include "Containers/Ticker.h"
+
+#include "EditorUtilityTest.generated.h"
+
+UENUM(BlueprintType)
+enum class EEditorUtilityTestResult : uint8
+{
+	/**
+	 * When finishing a test if you use Default, you're not explicitly stating if the test passed or failed.
+	 * Instead you're allowing any reported failure to have decided that for you. Even if you do
+	 * explicitly set as a success, it can be overturned by errors that occurred during the test.
+	 */
+	Default,
+	Invalid,
+	Preparing,
+	Running,
+	Failed,
+	Succeeded
+};
+
+DECLARE_DYNAMIC_MULTICAST_DELEGATE(FEditorUtilityTestEventSignature);
+DECLARE_DYNAMIC_MULTICAST_DELEGATE_OneParam(FEditorUtilityTestFinishedSignature, EEditorUtilityTestResult, TestState);
+
+UCLASS(Blueprintable, meta = (ShowWorldContextPin))
+class UEditorUtilityTest : public UObject
+{
+	GENERATED_UCLASS_BODY()
+
+public:
+
+	/** Called when the test is ready to prepare */
+	UPROPERTY(BlueprintAssignable)
+	FEditorUtilityTestEventSignature OnTestPrepare;
+
+	/** Called when the test is started */
+	UPROPERTY(BlueprintAssignable)
+	FEditorUtilityTestEventSignature OnTestStart;
+
+	/** Called when the test is finished. Use it to clean up */
+	UPROPERTY(BlueprintAssignable)
+	FEditorUtilityTestFinishedSignature OnTestFinished;
+
+	/**
+	 * The owner is the group or person responsible for the test. Generally you should use a group name
+	 * like 'Editor' or 'Rendering'. When a test fails it may not be obvious who should investigate
+	 * so this provides a associate responsible groups with tests.
+	 */
+	UPROPERTY(EditAnywhere, BlueprintReadOnly, Category = "Context", meta = (MultiLine = "true"))
+	FString Owner;
+
+	/** A description of the test, like what is this test trying to determine. */
+	UPROPERTY(EditAnywhere, BlueprintReadOnly, Category = "Context", meta = (MultiLine = "true"))
+	FString Description;
+
+	/** The Test's time limit for preparation, this is the time it has to trigger IsReadyToStart(). '0' means no limit. */
+	UPROPERTY(EditAnywhere, BlueprintReadOnly, Category = "Timeout")
+	float PreparationTimeLimit;
+
+	/** Test's total run time limit. '0' means no limit */
+	UPROPERTY(EditAnywhere, BlueprintReadOnly, Category = "Timeout")
+	float TimeLimit;
+
+protected:
+	/** Blueprint Utility Editor entry point */
+	UFUNCTION(BlueprintCallable, Category = "Editor Utility Test")
+	void Run();
+
+	/** Use to setup test before running, must call FinishPrepareTest at the end to actual start the test. */
+	UFUNCTION(BlueprintNativeEvent, Category = "Editor Utility Test")
+	void PrepareTest();
+
+	/** Tell the blueprint VM to start the test (to use at the end of Prepare Test event). */
+	UFUNCTION(BlueprintCallable, Category = "Editor Utility Test")
+	virtual void FinishPrepareTest();
+
+	/** Actual run the test, must call FinishTest at then end of test with a state to signify the test is done. */
+	UFUNCTION(BlueprintNativeEvent, Category = "Editor Utility Test")
+	void StartTest();
+
+	/** Tell the VM the test is finished with specify state. */
+	UFUNCTION(BlueprintCallable, Category = "Editor Utility Test", meta = (AdvancedDisplay = "Message"))
+	virtual void FinishTest(EEditorUtilityTestResult TestState, const FString& Message);
+
+	/** Use to add clean up steps, the call is blocking. */
+	UFUNCTION(BlueprintImplementableEvent, Category = "Editor Utility Test", meta = (DisplayName = "Finished Test"))
+	void ReceiveFinishedTest(EEditorUtilityTestResult TestState, EEditorUtilityTestResult& FinalState);
+
+	/** Add Error */
+	UFUNCTION(BlueprintCallable, Category = "Editor Utility Test")
+	virtual void AddError(const FString& Message);
+
+	/** Add Warning */
+	UFUNCTION(BlueprintCallable, Category = "Editor Utility Test")
+	virtual void AddWarning(const FString& Message);
+
+	/** Add Info */
+	UFUNCTION(BlueprintCallable, Category = "Editor Utility Test")
+	virtual void AddInfo(const FString& Message);
+
+	/** Add error if expected condition is false */
+	UFUNCTION(BlueprintCallable, Category = "Editor Utility Test", meta = (AdvancedDisplay = "ErrorMessage"))
+	virtual void ExpectTrue(bool Condition, const FString& ErrorMessage);
+
+	/** Add error if expected condition is true */
+	UFUNCTION(BlueprintCallable, Category = "Editor Utility Test", meta = (AdvancedDisplay = "ErrorMessage"))
+	virtual void ExpectFalse(bool Condition, const FString& ErrorMessage);
+
+	virtual bool Tick(float DeltaSeconds);
+
+	EEditorUtilityTestResult State;
+
+private:
+	/** Delegate for callbacks to Tick */
+	FTickerDelegate TickDelegate;
+
+	/** Handle to various registered delegates */
+	FTSTicker::FDelegateHandle TickDelegateHandle;
+
+	FDateTime StartTime;
+
+public:
+	/** Get test state */
+	UFUNCTION(BlueprintCallable, Category = "Editor Utility Test")
+	virtual EEditorUtilityTestResult GetState();
+
+	/** Is test ruuning */
+	UFUNCTION(BlueprintCallable, Category = "Editor Utility Test")
+	virtual bool IsRunning();
+
+};
diff --git a/Engine/Plugins/Tests/EditorTests/Source/EditorTests/EditorTests.Build.cs b/Engine/Plugins/Tests/EditorTests/Source/EditorTests/EditorTests.Build.cs
--- a/Engine/Plugins/Tests/EditorTests/Source/EditorTests/EditorTests.Build.cs
+++ b/Engine/Plugins/Tests/EditorTests/Source/EditorTests/EditorTests.Build.cs
@@ -69,6 +69,7 @@ public class EditorTests : ModuleRules
                 "MeshBuilder",
                 "RawMesh",
 				"AutomationController",
+				"Blutility",
 				//"SubobjectDataInterface",
 			}
 		);
diff --git a/Engine/Plugins/Tests/EditorTests/Source/EditorTests/Private/EditorUtilityTest.cpp b/Engine/Plugins/Tests/EditorTests/Source/EditorTests/Private/EditorUtilityTest.cpp
new file mode 100644
--- /dev/null
+++ b/Engine/Plugins/Tests/EditorTests/Source/EditorTests/Private/EditorUtilityTest.cpp
@@ -0,0 +1,327 @@
+// Copyright Epic Games, Inc. All Rights Reserved.
+
+#include "EditorUtilityTest.h"
+#include "Misc/AutomationTest.h"
+#include "EditorUtilityBlueprint.h"
+#include "EditorUtilitySubsystem.h"
+#include "AssetRegistry/AssetRegistryModule.h"
+#include "IAutomationControllerModule.h"
+#include "Misc/Paths.h"
+#include "Engine/Engine.h"
+#include "UnrealEdMisc.h"
+
+/////////////////////////////////////////////////////
+
+DEFINE_LOG_CATEGORY_STATIC(LogEditorUtilityTest, Log, All);
+
+class FBlueprintEditorUtilitiesTestBase : public FAutomationTestBase
+{
+public:
+	FBlueprintEditorUtilitiesTestBase(const FString& InName, const bool bInComplexTask)
+		: FAutomationTestBase(InName, bInComplexTask)
+		, EditorUtility(nullptr)
+	{
+	}
+
+	void SetEditorUtility(UEditorUtilityTest* InEditorUtility)
+	{
+		EditorUtility = InEditorUtility;
+	}
+
+	UEditorUtilityTest* GetEditorUtility()
+	{
+		return EditorUtility;
+	}
+
+	virtual FString GetTestOpenCommand(const FString& BlueprintAssetPath) const override
+	{
+		return FString::Printf(TEXT("Automate.OpenMap %s"), *BlueprintAssetPath);
+	}
+
+private:
+	UEditorUtilityTest* EditorUtility;
+};
+
+inline FBlueprintEditorUtilitiesTestBase* GetTest()
+{
+	return static_cast<FBlueprintEditorUtilitiesTestBase*>(FAutomationTestFramework::Get().GetCurrentTest());
+}
+
+UEditorUtilityTest::UEditorUtilityTest(const FObjectInitializer& ObjectInitializer)
+	: Super(ObjectInitializer)
+	, PreparationTimeLimit(15.0f)
+	, TimeLimit(60.0f)
+	, State(EEditorUtilityTestResult::Invalid)
+	, StartTime(0)
+{
+	TickDelegate = FTickerDelegate::CreateUObject(this, &UEditorUtilityTest::Tick);
+}
+
+void UEditorUtilityTest::Run()
+{
+	if (State != EEditorUtilityTestResult::Preparing || State != EEditorUtilityTestResult::Running)
+	{
+		State = EEditorUtilityTestResult::Preparing;
+
+		StartTime = FDateTime::Now();
+		if (PreparationTimeLimit || TimeLimit)
+		{
+			TickDelegateHandle = FTSTicker::GetCoreTicker().AddTicker(TickDelegate, 1.f);
+		}
+
+		FBlueprintEditorUtilitiesTestBase* CurrentTest = GetTest();
+		if (CurrentTest != nullptr)
+		{
+			CurrentTest->SetEditorUtility(this);
+		}
+
+		PrepareTest();
+		OnTestPrepare.Broadcast();
+	}
+}
+
+void UEditorUtilityTest::PrepareTest_Implementation()
+{
+	FinishPrepareTest();
+}
+
+void UEditorUtilityTest::FinishPrepareTest()
+{
+	if (State == EEditorUtilityTestResult::Preparing)
+	{
+		State = EEditorUtilityTestResult::Running;
+		StartTest();
+		OnTestStart.Broadcast();
+	}
+}
+
+void UEditorUtilityTest::StartTest_Implementation()
+{
+	FinishTest(EEditorUtilityTestResult::Default, TEXT(""));
+}
+
+void UEditorUtilityTest::FinishTest(EEditorUtilityTestResult ProposedState, const FString& Message)
+{
+	if (State == EEditorUtilityTestResult::Running)
+	{
+		if (TickDelegateHandle.IsValid())
+		{
+			TickDelegateHandle.Reset();
+		}
+
+		if (ProposedState > EEditorUtilityTestResult::Running)
+		{
+			State = ProposedState;
+		}
+		else
+		{
+			State = EEditorUtilityTestResult::Succeeded;
+		}
+
+		EEditorUtilityTestResult NewState = EEditorUtilityTestResult::Default;
+		ReceiveFinishedTest(State, NewState);
+		if (NewState > EEditorUtilityTestResult::Running)
+		{
+			State = NewState;
+		}
+
+		if (State == EEditorUtilityTestResult::Succeeded && !Message.IsEmpty())
+		{
+			AddInfo(Message);
+		}
+		else if (State == EEditorUtilityTestResult::Failed)
+		{
+			AddError(Message.IsEmpty() ? TEXT("Test finished with a Failed state.") : Message);
+		}
+
+		FBlueprintEditorUtilitiesTestBase* CurrentTest = GetTest();
+		if (CurrentTest != nullptr && CurrentTest->GetEditorUtility() == this)
+		{
+			// Reset blueprint test object tracking
+			CurrentTest->SetEditorUtility(nullptr);
+		}
+
+		OnTestFinished.Broadcast(State);
+
+		// Force GC at the end of every test.
+		GEngine->ForceGarbageCollection();
+	}
+}
+
+EEditorUtilityTestResult UEditorUtilityTest::GetState()
+{
+	return State;
+}
+
+bool UEditorUtilityTest::IsRunning()
+{
+	return State == EEditorUtilityTestResult::Preparing || State == EEditorUtilityTestResult::Running;
+}
+
+bool UEditorUtilityTest::Tick(float DeltaSeconds)
+{
+	float RunTime = (FDateTime::Now() - StartTime).GetTotalSeconds();
+	if ((State == EEditorUtilityTestResult::Preparing && PreparationTimeLimit && RunTime > PreparationTimeLimit)
+		|| (State == EEditorUtilityTestResult::Running && TimeLimit && RunTime > TimeLimit))
+	{
+		FString Message = FString::Printf(TEXT("Test reached timeout after %.01f seconds during %s phase."), RunTime, *UEnum::GetDisplayValueAsText(State).ToString());
+		FinishTest(EEditorUtilityTestResult::Failed, Message);
+	}
+
+	return true;
+}
+
+void UEditorUtilityTest::AddError(const FString& Message)
+{
+	FBlueprintEditorUtilitiesTestBase* CurrentTest = GetTest();
+	if (CurrentTest != nullptr && CurrentTest->GetEditorUtility() == this)
+	{
+		CurrentTest->AddError(Message);
+	}
+	else
+	{
+		UE_LOG(LogEditorUtilityTest, Error, TEXT("%s"), *Message);
+	}
+}
+
+void UEditorUtilityTest::AddWarning(const FString& Message)
+{
+	FBlueprintEditorUtilitiesTestBase* CurrentTest = GetTest();
+	if (CurrentTest != nullptr && CurrentTest->GetEditorUtility() == this)
+	{
+		CurrentTest->AddWarning(Message);
+	}
+	else
+	{
+		UE_LOG(LogEditorUtilityTest, Warning, TEXT("%s"), *Message);
+	}
+}
+
+void UEditorUtilityTest::AddInfo(const FString& Message)
+{
+	FBlueprintEditorUtilitiesTestBase* CurrentTest = GetTest();
+	if (CurrentTest != nullptr && CurrentTest->GetEditorUtility() == this)
+	{
+		CurrentTest->AddInfo(Message);
+	}
+	else
+	{
+		UE_LOG(LogEditorUtilityTest, Log, TEXT("%s"), *Message);
+	}
+}
+
+void UEditorUtilityTest::ExpectTrue(bool Condition, const FString& ErrorMessage)
+{
+	if (!Condition)
+	{
+		AddError(ErrorMessage.IsEmpty() ? TEXT("Expected condition was not True.") : ErrorMessage);
+	}
+}
+
+void UEditorUtilityTest::ExpectFalse(bool Condition, const FString& ErrorMessage)
+{
+	if (Condition)
+	{
+		AddError(ErrorMessage.IsEmpty()? TEXT("Expected condition was not False.") : ErrorMessage);
+	}
+}
+
+class FWaitForBlueprintEditorUtilityTest : public IAutomationLatentCommand
+{
+public:
+	FWaitForBlueprintEditorUtilityTest()
+	{}
+	virtual bool Update() override
+	{
+		FBlueprintEditorUtilitiesTestBase* CurrentTest = GetTest();
+		if (CurrentTest == nullptr)
+		{
+			return true;
+		}
+		UEditorUtilityTest* EditorUtility = CurrentTest->GetEditorUtility();
+
+		if (EditorUtility == nullptr)
+		{
+			return true;
+		}
+
+		return !EditorUtility->IsRunning();
+	}
+};
+
+IMPLEMENT_CUSTOM_COMPLEX_AUTOMATION_TEST(FBlueprintEditorUtilitiesTest, FBlueprintEditorUtilitiesTestBase, "Project.Blueprints.EditorUtilities", EAutomationTestFlags::EditorContext | EAutomationTestFlags::ProductFilter)
+
+/**
+ * Gather the tests to run
+ */
+void FBlueprintEditorUtilitiesTest::GetTests(TArray<FString>& OutBeautifiedNames, TArray<FString>& OutTestCommands) const
+{
+	IAutomationControllerModule& AutomationControllerModule = FModuleManager::LoadModuleChecked<IAutomationControllerModule>(TEXT("AutomationController"));
+	IAutomationControllerManagerPtr AutomationController = AutomationControllerModule.GetAutomationController();
+	bool IsDeveloperDirectoryIncluded = AutomationController->IsDeveloperDirectoryIncluded();
+
+	FTopLevelAssetPath BaseClassPathName = UEditorUtilityTest::StaticClass()->GetClassPathName();
+	UClass* BlueprintClass = UEditorUtilityBlueprint::StaticClass();
+
+	FAssetRegistryModule& AssetRegistryModule = FModuleManager::LoadModuleChecked<FAssetRegistryModule>(TEXT("AssetRegistry"));
+
+	// Use the asset registry to get the set of all class names deriving from Base
+	TSet< FTopLevelAssetPath > DerivedPathNames;
+	{
+		TArray< FTopLevelAssetPath > BasePathNames;
+		BasePathNames.Add(BaseClassPathName);
+
+		TSet< FTopLevelAssetPath > Excluded;
+		AssetRegistryModule.Get().GetDerivedClassNames(BasePathNames, Excluded, DerivedPathNames);
+	}
+
+	TArray<FAssetData> ObjectList;
+	AssetRegistryModule.Get().GetAssetsByClass(BlueprintClass->GetClassPathName(), ObjectList);
+	for (auto ObjIter = ObjectList.CreateConstIterator(); ObjIter; ++ObjIter)
+	{
+		const FAssetData& Asset = *ObjIter;
+		FString Filename = Asset.ObjectPath.ToString();
+
+		// Get the the class this blueprint generates
+		auto AssetTag = Asset.TagsAndValues.FindTag(TEXT("GeneratedClass"));
+		if (AssetTag.IsSet())
+		{
+			// Convert path to just the name part
+			const FTopLevelAssetPath ClassPathName(FPackageName::ExportTextPathToObjectPath(*AssetTag.AsString()));
+
+			// Check if this class is in the derived set
+			if (!DerivedPathNames.Contains(ClassPathName))
+			{
+				continue;
+			}
+
+			// Exclude Developer folder
+			if (Filename.Find(TEXT("/Game/")) == 0)
+			{
+				if (!IsDeveloperDirectoryIncluded && Filename.Find(TEXT("/Game/Developers")) == 0) continue;
+
+				// Extract folder org
+				Filename = FPaths::GetPath(Filename.RightChop(6)).Replace(TEXT("/"), TEXT("."));
+				FString BeautifiedFilename = Filename + TEXT(".") + Asset.AssetName.ToString();
+
+				OutBeautifiedNames.Add(BeautifiedFilename);
+				OutTestCommands.Add(Asset.ObjectPath.ToString());
+			}
+		}
+	}
+}
+
+/**
+ * Runs Blueprint Editor Utilities Tests
+ */
+bool FBlueprintEditorUtilitiesTest::RunTest(const FString& BlueprintAssetPath)
+{
+	UObject* Object = StaticLoadObject(UObject::StaticClass(), NULL, *BlueprintAssetPath);
+
+	UEditorUtilitySubsystem* EditorUtilitySubsystem = GEditor->GetEditorSubsystem<UEditorUtilitySubsystem>();
+	EditorUtilitySubsystem->TryRun(Object);
+
+	ADD_LATENT_AUTOMATION_COMMAND(FWaitForBlueprintEditorUtilityTest());
+
+	return true;
+}
