Fix Static analyzer warnings in EditorUtilityTest.cpp
#hordeIssue 202787
#rnx

diff --git a/Engine/Plugins/Tests/EditorTests/Source/EditorTests/Private/EditorUtilityTest.cpp b/Engine/Plugins/Tests/EditorTests/Source/EditorTests/Private/EditorUtilityTest.cpp
--- a/Engine/Plugins/Tests/EditorTests/Source/EditorTests/Private/EditorUtilityTest.cpp
+++ b/Engine/Plugins/Tests/EditorTests/Source/EditorTests/Private/EditorUtilityTest.cpp
@@ -59,7 +59,7 @@ UEditorUtilityTest::UEditorUtilityTest(const FObjectInitializer& ObjectInitializ
 
 void UEditorUtilityTest::Run()
 {
-	if (State != EEditorUtilityTestResult::Preparing || State != EEditorUtilityTestResult::Running)
+	if (State != EEditorUtilityTestResult::Preparing && State != EEditorUtilityTestResult::Running)
 	{
 		State = EEditorUtilityTestResult::Preparing;
 
