Clean up some deprecated blueprint tests
+ fix SkeletalCubeNoAnim_Blueprint Physic asset target

#rnx

diff --git a/Engine/Plugins/Tests/EditorTests/Source/EditorTests/Private/UnrealEd/BlueprintAutomationTests.cpp b/Engine/Plugins/Tests/EditorTests/Source/EditorTests/Private/UnrealEd/BlueprintAutomationTests.cpp
--- a/Engine/Plugins/Tests/EditorTests/Source/EditorTests/Private/UnrealEd/BlueprintAutomationTests.cpp
+++ b/Engine/Plugins/Tests/EditorTests/Source/EditorTests/Private/UnrealEd/BlueprintAutomationTests.cpp
@@ -48,10 +48,6 @@
 
 DEFINE_LOG_CATEGORY_STATIC(LogBlueprintAutomationTests, Log, All);
 
-IMPLEMENT_COMPLEX_AUTOMATION_TEST(FBlueprintCompileOnLoadTest, "Project.Blueprints.Compile-On-Load", EAutomationTestFlags::EditorContext | EAutomationTestFlags::StressFilter)
-IMPLEMENT_COMPLEX_AUTOMATION_TEST(FBlueprintInstancesTest, "Project.Blueprints.Instance Test", EAutomationTestFlags::EditorContext | EAutomationTestFlags::StressFilter)
-IMPLEMENT_COMPLEX_AUTOMATION_TEST(FBlueprintReparentTest, "System.Blueprints.Reparent", EAutomationTestFlags::EditorContext | EAutomationTestFlags::StressFilter)
-IMPLEMENT_COMPLEX_AUTOMATION_TEST(FBlueprintRenameAndCloneTest, "Project.Blueprints.Rename And Clone", (EAutomationTestFlags::EditorContext | EAutomationTestFlags::StressFilter))
 IMPLEMENT_COMPLEX_AUTOMATION_TEST(FCompileBlueprintsTest, "Project.Blueprints.Compile Blueprints", EAutomationTestFlags::EditorContext | EAutomationTestFlags::StressFilter)
 IMPLEMENT_COMPLEX_AUTOMATION_TEST(FCompileAnimBlueprintsTest, "Project.Blueprints.Compile Anims", EAutomationTestFlags::EditorContext | EAutomationTestFlags::StressFilter)
 
@@ -830,216 +826,6 @@ private:
 	bool bRunGCOnClose;
 };
 
-/************************************************************************/
-/* FBlueprintCompileOnLoadTest                                          */
-/************************************************************************/
-
-/** 
- * Gather the tests to run
- */
-void FBlueprintCompileOnLoadTest::GetTests(TArray<FString>& OutBeautifiedNames, TArray<FString>& OutTestCommands) const
-{
-	bool bTestLoadedBlueprints = false;
-	GConfig->GetBool( TEXT("AutomationTesting.Blueprint"), TEXT("TestAllBlueprints"), /*out*/ bTestLoadedBlueprints, GEngineIni );
-	FBlueprintAutomationTestUtilities::CollectTestsByClass(UBlueprint::StaticClass(), OutBeautifiedNames, OutTestCommands, !bTestLoadedBlueprints);
-}
-
-/** 
- * Runs compile-on-load test against all unloaded, and optionally loaded, blueprints
- * See the TestAllBlueprints config key in the [Automation.Blueprint] config sections
- */
-bool FBlueprintCompileOnLoadTest::RunTest(const FString& BlueprintAssetPath)
-{
-	FCompilerResultsLog Results;
-
-	UBlueprint* ExistingBP = nullptr;
-	// if this blueprint was already loaded, then these tests are invalidated 
-	// (because dependencies have already been loaded)
-	if (FBlueprintAutomationTestUtilities::IsBlueprintLoaded(BlueprintAssetPath, &ExistingBP))
-	{
-		if (FBlueprintAutomationTestUtilities::IsAssetUnsaved(BlueprintAssetPath))
-		{
-			AddError(FString::Printf(TEXT("You have unsaved changes made to '%s', please save them before running this test."), *BlueprintAssetPath));
-			return false;
-		}
-		else
-		{
-			AddWarning(FString::Printf(TEXT("Test may be invalid (the blueprint is already loaded): '%s'"), *BlueprintAssetPath));
-			FBlueprintAutomationTestUtilities::UnloadBlueprint(ExistingBP);
-		}
-	}
-
-	// tracks blueprints that were already loaded (and cleans up any that were 
-	// loaded in its lifetime, once it is destroyed)
-	FScopedBlueprintUnloader NewBlueprintUnloader(/*bAutoOpenScope =*/true, /*bRunGCOnCloseIn =*/true);
-
-	// We load the blueprint twice and compare the two for discrepancies. This is 
-	// to bring dependency load issues to light (among other things). If a blueprint's
-	// dependencies are loaded too late, then this first object is the degenerate one.
-	UBlueprint* InitialBlueprint = Cast<UBlueprint>(StaticLoadObject(UBlueprint::StaticClass(), NULL, *BlueprintAssetPath));
-
-	// if we failed to load it the first time, then there is no need to make a 
-	// second attempt, leave them to fix up this issue first
-	if (InitialBlueprint == NULL)
-	{
-		AddError(*FString::Printf(TEXT("Unable to load blueprint for: '%s'"), *BlueprintAssetPath));
-		return false;
-	}
-
-	if (!InitialBlueprint->SkeletonGeneratedClass || !InitialBlueprint->GeneratedClass)
-	{
-		AddError(*FString::Printf(TEXT("Unable to load blueprint for: '%s'. Probably it derives from an invalid class."), *BlueprintAssetPath));
-		return false;
-	}
-
-	// GATHER SUBOBJECTS
-	TArray<TWeakObjectPtr<UObject>> InitialBlueprintSubobjects;
-	{
-		TArray<UObject*> InitialBlueprintSubobjectsPtr;
-		GetObjectsWithOuter(InitialBlueprint, InitialBlueprintSubobjectsPtr);
-		for (UObject* Obj : InitialBlueprintSubobjectsPtr)
-		{
-			InitialBlueprintSubobjects.Add(Obj);
-		}
-	}
-
-	// GATHER DEPENDENCIES
-	TSet<TWeakObjectPtr<UBlueprint>> BlueprintDependencies;
-	{
-		TArray<UBlueprint*> DependentBlueprints;
-		FBlueprintEditorUtils::FindDependentBlueprints(InitialBlueprint, DependentBlueprints);
-		for (UBlueprint* BP : DependentBlueprints)
-		{
-			BlueprintDependencies.Add(BP);
-		}
-	}
-	BlueprintDependencies.Add(InitialBlueprint);
-
-	// GATHER DEPENDENCIES PERSISTENT DATA
-	struct FReplaceInnerData
-	{
-		TWeakObjectPtr<UClass> Class;
-		FSoftObjectPath BlueprintAsset;
-	};
-	TArray<FReplaceInnerData> ReplaceInnerData;
-	for (TWeakObjectPtr<UBlueprint>& BPToUnloadWP : BlueprintDependencies)
-	{
-		UBlueprint* BPToUnload = BPToUnloadWP.Get();
-		UClass* OldClass = BPToUnload ? *BPToUnload->GeneratedClass : nullptr;
-		if (OldClass)
-		{
-			FReplaceInnerData Data;
-			Data.Class = OldClass;
-			Data.BlueprintAsset = FSoftObjectPath(BPToUnload);
-			ReplaceInnerData.Add(Data);
-		}
-	}
-
-	// store off data for the initial blueprint so we can unload it (and reconstruct 
-	// later to compare it with a second one)
-	TArray<uint8> InitialLoadData;
-	FObjectWriter(InitialBlueprint, InitialLoadData);
-
-	// grab the name before we unload the blueprint
-	FName const BlueprintName = InitialBlueprint->GetFName();
-	// unload the blueprint so we can reload it (to catch any differences, now  
-	// that all its dependencies should be loaded as well)
-
-	//UNLOAD DEPENDENCIES, all circular dependencies will be loaded again 
-	// unload the blueprint so we can reload it (to catch any differences, now  
-	// that all its dependencies should be loaded as well)
-	for (TWeakObjectPtr<UBlueprint> BPToUnloadWP : BlueprintDependencies)
-	{
-		if (UBlueprint* BPToUnload = BPToUnloadWP.Get())
-		{
-			FBlueprintAutomationTestUtilities::UnloadBlueprint(BPToUnload);
-		}
-	}
-
-	// this blueprint is now dead (will be destroyed next garbage-collection pass)
-	UBlueprint* UnloadedBlueprint = InitialBlueprint;
-	InitialBlueprint = NULL;
-
-	// load the blueprint a second time; if the two separately loaded blueprints 
-	// are different, then this one is most likely the choice one (it has all its 
-	// dependencies loaded)
-
-	UBlueprint* ReloadedBlueprint = Cast<UBlueprint>(StaticLoadObject(UBlueprint::StaticClass(), NULL, *BlueprintAssetPath));
-
-	UPackage* TransientPackage = GetTransientPackage();
-	FName ReconstructedName = MakeUniqueObjectName(TransientPackage, UBlueprint::StaticClass(), BlueprintName);
-	// reconstruct the initial blueprint (using the serialized data from its initial load)
-	EObjectFlags const StandardBlueprintFlags = RF_Public | RF_Standalone | RF_Transactional;
-	InitialBlueprint = NewObject<UBlueprint>(TransientPackage, ReconstructedName, StandardBlueprintFlags | RF_Transient);
-	FObjectReader(InitialBlueprint, InitialLoadData);
-	{
-		TMap<UObject*, UObject*> ClassRedirects;
-		for (FReplaceInnerData& Data : ReplaceInnerData)
-		{
-			UClass* OriginalClass = Data.Class.Get();
-			UBlueprint* NewBlueprint = Cast<UBlueprint>(Data.BlueprintAsset.ResolveObject());
-			UClass* NewClass = NewBlueprint ? *NewBlueprint->GeneratedClass : NULL;
-			if (OriginalClass && NewClass)
-			{
-				ClassRedirects.Add(OriginalClass, NewClass);
-			}
-		}
-		// REPLACE OLD DATA
-		FArchiveReplaceObjectRef<UObject>(InitialBlueprint, ClassRedirects, EArchiveReplaceObjectFlags::IgnoreOuterRef);
-		for (TWeakObjectPtr<UObject>& SubobjWP : InitialBlueprintSubobjects)
-		{
-			if (UObject* Subobj = SubobjWP.Get())
-			{
-				FArchiveReplaceObjectRef<UObject>(Subobj, ClassRedirects, EArchiveReplaceObjectFlags::IgnoreOuterRef);
-			}
-		}
-	}
-
-	// look for diffs between subsequent loads and log them as errors
-	TArray<FDiffSingleResult> BlueprintDiffs;
-	bool bDiffsFound = FBlueprintAutomationTestUtilities::DiffBlueprints(InitialBlueprint, ReloadedBlueprint, BlueprintDiffs);
-	if (bDiffsFound)
-	{
-		FBlueprintAutomationTestUtilities::ResolveCircularDependencyDiffs(ReloadedBlueprint, BlueprintDiffs);
-		// if there are still diffs after resolving any the could have been from unloaded circular dependencies
-		if (BlueprintDiffs.Num() > 0)
-		{
-			AddError(FString::Printf(TEXT("Inconsistencies between subsequent blueprint loads for: '%s' (was a dependency not preloaded?)"), *BlueprintAssetPath));
-		}
-		else 
-		{
-			bDiffsFound = false;
-		}
-		
-		// list all the differences (so as to help identify what dependency was missing)
-		for (auto DiffIt(BlueprintDiffs.CreateIterator()); DiffIt; ++DiffIt)
-		{
-			// will be presented in the context of "what changed between the initial load and the second?"
-			FString DiffDescription = DiffIt->ToolTip.ToString();
-			if (DiffDescription != DiffIt->DisplayString.ToString())
-			{
-				DiffDescription = FString::Printf(TEXT("%s (%s)"), *DiffDescription, *DiffIt->DisplayString.ToString());
-			}
-
-			const UEdGraphNode* NodeFromPin = DiffIt->Pin1 ? DiffIt->Pin1->GetOuter() : nullptr;
-			const UEdGraphNode* Node = DiffIt->Node1 ? DiffIt->Node1 : NodeFromPin;
-			const UEdGraph* Graph = Node ? Node->GetGraph() : NULL;
-			const FString GraphName = Graph ? Graph->GetName() : FString(TEXT("Unknown Graph"));
-			AddError(FString::Printf(TEXT("%s.%s differs between subsequent loads: %s"), *BlueprintName.ToString(), *GraphName, *DiffDescription));
-		}
-	}
-
-	// At the close of this function, the FScopedBlueprintUnloader should prep 
-	// for following tests by unloading any blueprint dependencies that were 
-	// loaded for this one (should catch InitialBlueprint and ReloadedBlueprint) 
-	// 
-	// The FScopedBlueprintUnloader should also run garbage-collection after,
-	// in hopes that the imports for this blueprint get destroyed so that they 
-	// don't invalidate other tests that share the same dependencies
-	return !bDiffsFound;
-}
-
-
 /************************************************************************/
 /* FCompileBlueprintsTest                                              */
 /************************************************************************/
@@ -1072,316 +858,3 @@ bool FCompileAnimBlueprintsTest::RunTest(const FString& Parameters)
 	return FBlueprintAutomationTestUtilities::CompileBlueprint(Parameters);
 }
 
-/************************************************************************/
-/* FBlueprintInstancesTest                                              */
-/************************************************************************/
-
-void FBlueprintInstancesTest::GetTests(TArray<FString>& OutBeautifiedNames, TArray<FString>& OutTestCommands) const
-{
-	// Load the test maps
-	check( GConfig );
-
-	// Load from config file
-	TArray<FString> MapsToLoad;
-	GConfig->GetArray( TEXT("AutomationTesting.Blueprint"), TEXT("InstanceTestMaps"), MapsToLoad, GEngineIni);
-	for (auto It = MapsToLoad.CreateConstIterator(); It; ++It)
-	{
-		const FString MapFileName = *It;
-		if( IFileManager::Get().FileSize( *MapFileName ) > 0 )
-		{
-			OutBeautifiedNames.Add(FPaths::GetBaseFilename(MapFileName));
-			OutTestCommands.Add(MapFileName);
-		}
-	}
-}
-
-/**
-  * Wait for the given amount of time
-  */
-DEFINE_LATENT_AUTOMATION_COMMAND_ONE_PARAMETER(FDelayLatentCommand, float, Duration);
-
-bool FDelayLatentCommand::Update()
-{
-	float NewTime = FPlatformTime::Seconds();
-	if (NewTime - StartTime >= Duration)
-	{
-		return true;
-	}
-	return false;
-}
-
-/************************************************************************/
-/**
- * Uses test maps in Engine and/or game content folder which are populated with a few blueprint instances
- * See InstanceTestMaps entries in the [Automation.Blueprint] config sections
- * For all blueprint instances in the map:
- *		Duplicates the instance 
- *		Compares the duplicated instance properties to the original instance properties
- */
-bool FBlueprintInstancesTest::RunTest(const FString& InParameters)
-{
-	FBlueprintAutomationTestUtilities::LoadMap(InParameters);
-
-	// Pause before running test
-	ADD_LATENT_AUTOMATION_COMMAND(FDelayLatentCommand(2.f));
-
-	// Grab BP instances from map
-	TSet<AActor*> BlueprintInstances;
-	for ( FActorIterator It(GWorld); It; ++It )
-	{
-		AActor* Actor = *It;
-
-		UClass* ActorClass = Actor->GetClass();
-
-		if (ActorClass->ClassGeneratedBy && ActorClass->ClassGeneratedBy->IsA( UBlueprint::StaticClass() ) )
-		{
-			BlueprintInstances.Add(Actor);
-		}
-	}
-
-	bool bPropertiesMatch = true;
-	FCompilerResultsLog ResultLog;
-	TSet<UPackage*> PackagesUserRefusedToFullyLoad;
-	ObjectTools::FPackageGroupName PGN;
-
-	for (auto BpIter = BlueprintInstances.CreateIterator(); BpIter; ++BpIter )
-	{
-		AActor* BPInstance = *BpIter;
-		UObject* BPInstanceOuter = BPInstance->GetOuter();
-
-		TMap<FString,FString> BPNativePropertyValues;
-		BPInstance->GetNativePropertyValues(BPNativePropertyValues);
-
-		// Grab the package and save out its dirty state
-		UPackage* ActorPackage = BPInstance->GetOutermost();
-		FBlueprintAutomationTestUtilities::FPackageCleaner Cleaner(ActorPackage);
-
-		// Use this when duplicating the object to keep a list of everything that was duplicated
-		//TMap<UObject*, UObject*> DuplicatedObjectList;
-
-		FObjectDuplicationParameters Parameters(BPInstance, BPInstanceOuter);
-		//Parameters.CreatedObjects = &DuplicatedObjectList;
-		Parameters.DestName = MakeUniqueObjectName( BPInstanceOuter, AActor::StaticClass(), BPInstance->GetFName() );
-
-		// Duplicate the object
-		AActor* ClonedInstance = Cast<AActor>(StaticDuplicateObjectEx(Parameters));
-
-		if (!FBlueprintAutomationTestUtilities::CompareObjects(BPInstance, ClonedInstance, ResultLog))
-		{
-			bPropertiesMatch = false;
-			break;
-		}
-
-		// Ensure we can't save package in editor
-		FBlueprintAutomationTestUtilities::DontSavePackage(ActorPackage);
-	}
-
-	// Start a new map for now
-	// @todo find a way return to previous map thats a 100% reliably
-	GEditor->CreateNewMapForEditing();
-
-	ADD_LATENT_AUTOMATION_COMMAND(FDelayLatentCommand(2.f));
-
-	return bPropertiesMatch;
-}
-
-/*******************************************************************************
-* FBlueprintReparentTest
-*******************************************************************************/
-
-void FBlueprintReparentTest::GetTests(TArray<FString>& OutBeautifiedNames, TArray<FString>& OutTestCommands) const
-{
-	TArray<FAssetData> Assets;
-	FBlueprintAutomationTestUtilities::GetAssetListingFromConfig(TEXT("ReparentTest.ChildrenPackagePaths"), Assets, UBlueprint::StaticClass());
-
-	for (FAssetData const& AssetData : Assets)
-	{
-		OutBeautifiedNames.Add(AssetData.AssetName.ToString());
-		OutTestCommands.Add(AssetData.ObjectPath.ToString());
-	}
-}
-
-bool FBlueprintReparentTest::RunTest(const FString& BlueprintAssetPath)
-{
-	bool bTestFailed = false;
-
-	UBlueprint * const BlueprintTemplate = Cast<UBlueprint>(StaticLoadObject(UBlueprint::StaticClass(), NULL, *BlueprintAssetPath));
-	if (BlueprintTemplate != NULL)
-	{
-		// want to explicitly test switching from actors->objects, and vise versa (objects->actors), 
-		// also could cover the case of changing non-native parents to native ones
-		TArray<UClass*> TestParentClasses;
-		if (!BlueprintTemplate->ParentClass->IsChildOf(AActor::StaticClass()))
-		{
-			TestParentClasses.Add(AActor::StaticClass());
-		}
-		else
-		{
-			// not many engine level Blueprintable classes that aren't Actors
-			TestParentClasses.Add(USaveGame::StaticClass());
-		}
-
-		TArray<FAssetData> Assets;
-		FBlueprintAutomationTestUtilities::GetAssetListingFromConfig(TEXT("ReparentTest.ParentsPackagePaths"), Assets, UBlueprint::StaticClass());
-		// additionally gather up any blueprints that we explicitly specify though the config
-		for (FAssetData const& AssetData : Assets)
-		{
-			UClass* AssetClass = FindObject<UClass>(AssetData.AssetClassPath);
-			TestParentClasses.Add(AssetClass);
-		}
-
-		for (UClass* Class : TestParentClasses)
-		{		 
-			UBlueprint* BlueprintObj = FBlueprintAutomationTestUtilities::DuplicateBlueprint(BlueprintTemplate);
-			check(BlueprintObj != NULL);
-			BlueprintObj->ParentClass = Class;
-
-			if (!FBlueprintAutomationTestUtilities::TestSaveBlueprint(BlueprintObj))
-			{
-				AddError(FString::Printf(TEXT("Failed to save blueprint after reparenting with %s: '%s'"), *Class->GetName(), *BlueprintAssetPath));
-				bTestFailed = true;
-			}
-
-			FBlueprintAutomationTestUtilities::UnloadBlueprint(BlueprintObj);
-		}
-
-#if WITH_EDITOR
-		// clear undo history to ensure that the transaction buffer isn't 
-		// holding onto any references to the blueprints we want unloaded
-		GEditor->Trans->Reset(NSLOCTEXT("BpAutomation", "ReparentTest", "Reparent Blueprint Test"));
-#endif // #if WITH_EDITOR
-		// make sure the unloaded blueprints are properly flushed (for future tests)
-		CollectGarbage(RF_NoFlags);
-	}
-
-	return !bTestFailed;
-}
-
-/*******************************************************************************
-* FBlueprintRenameTest
-*******************************************************************************/
-
-void FBlueprintRenameAndCloneTest::GetTests(TArray<FString>& OutBeautifiedNames, TArray<FString>& OutTestCommands) const
-{
-	FAssetRegistryModule& AssetRegistryModule = FModuleManager::LoadModuleChecked<FAssetRegistryModule>(TEXT("AssetRegistry"));
-	
-	TArray<FAssetData> ObjectList;
-	AssetRegistryModule.Get().GetAssetsByClass(UBlueprint::StaticClass()->GetClassPathName(), ObjectList);
-
-	for (FAssetData const& Asset : ObjectList)
-	{
-		FString AssetObjPath = Asset.ObjectPath.ToString();
-
-		FString const Filename = FPackageName::LongPackageNameToFilename(AssetObjPath);
-		if (!FAutomationTestFramework::Get().ShouldTestContent(Filename))
-		{
-			continue;
-		}
-
-		FString PackageName, AssetName;
-		FBlueprintAutomationTestUtilities::SplitPackagePathAndAssetName(AssetObjPath, PackageName, AssetName);
-
-		if (UPackage* ExistingPackage = FindPackage(NULL, *PackageName))
-		{
-			if (ExistingPackage->IsRooted())
-			{
-				continue;
-			}
-		}
-
-		OutBeautifiedNames.Add(Asset.AssetName.ToString());
-		OutTestCommands.Add(AssetObjPath);
-	}
-}
-
-bool FBlueprintRenameAndCloneTest::RunTest(const FString& BlueprintAssetPath)
-{
-	bool bTestFailed = false;
-	if (FBlueprintAutomationTestUtilities::IsAssetUnsaved(BlueprintAssetPath))
-	{
-		bTestFailed = true;
-		AddError(FString::Printf(TEXT("You have unsaved changes made to '%s', please save them before running this test."), *BlueprintAssetPath));
-	}
-
-	bool bIsAlreadyLoaded = false;
-	if (FBlueprintAutomationTestUtilities::IsBlueprintLoaded(BlueprintAssetPath))
-	{
-		bIsAlreadyLoaded = true;
-		AddWarning(FString::Printf(TEXT("'%s' is already loaded, and possibly referenced by external objects (unable to perform rename tests... please run again in an empty map)."), *BlueprintAssetPath));
-	}
-
-	// track the loaded blueprint (and any other BP dependencies) so we can 
-	// unload them if we end up renaming it.
-	FScopedBlueprintUnloader NewBlueprintUnloader(/*bAutoOpenScope =*/true, /*bRunGCOnCloseIn =*/false);
-
-	UBlueprint* const OriginalBlueprint = Cast<UBlueprint>(StaticLoadObject(UBlueprint::StaticClass(), NULL, *BlueprintAssetPath));
-	if (OriginalBlueprint == NULL)
-	{
-		bTestFailed = true;
-		AddError(FString::Printf(TEXT("Failed to load '%s' (has it been renamed?)."), *BlueprintAssetPath));
-	}
-	else if (!OriginalBlueprint->SkeletonGeneratedClass || !OriginalBlueprint->GeneratedClass)
-	{
-		bTestFailed = true;
-		AddError(*FString::Printf(TEXT("Unable to load blueprint for: '%s'. Probably it derives from an invalid class."), *BlueprintAssetPath));
-	}
-	else if (!bTestFailed)
-	{
-		// duplicate
-		{
-			UBlueprint* DuplicateBluprint = FBlueprintAutomationTestUtilities::DuplicateBlueprint(OriginalBlueprint);
-			if (!FBlueprintAutomationTestUtilities::TestSaveBlueprint(DuplicateBluprint))
-			{
-				AddError(FString::Printf(TEXT("Failed to save blueprint after duplication: '%s'"), *BlueprintAssetPath));
-				bTestFailed = true;
-			}
-			FBlueprintAutomationTestUtilities::UnloadBlueprint(DuplicateBluprint);
-		}
-
-		// rename
-		if (!bIsAlreadyLoaded)
-		{
-			// store the original package so we can manually invalidate it after the move
-			UPackage* const OriginalPackage = OriginalBlueprint->GetOutermost();
-
-			FString const BlueprintName = OriginalBlueprint->GetName();
-			UPackage* TempPackage = FBlueprintAutomationTestUtilities::CreateTempPackage(BlueprintName);
-
-			FString NewName = FString::Printf(TEXT("%s-Rename"), *BlueprintName);
-			NewName = MakeUniqueObjectName(TempPackage, OriginalBlueprint->GetClass(), *NewName).ToString();
-
-			OriginalBlueprint->Rename(*NewName, TempPackage, REN_None);
-
-			if (!FBlueprintAutomationTestUtilities::TestSaveBlueprint(OriginalBlueprint))
-			{
-				AddError(FString::Printf(TEXT("Failed to save blueprint after rename: '%s'"), *BlueprintAssetPath));
-				bTestFailed = true;
-			}
-
-			// the blueprint has been moved out of this package, invalidate it so 
-			// we don't save it in this state and so we can reload the blueprint later
-			FBlueprintAutomationTestUtilities::InvalidatePackage(OriginalPackage);
-
-			// need to unload the renamed blueprint (and any other blueprints 
-			// that were relying on it), so that the renamed blueprint doesn't get used by the user
-			FBlueprintAutomationTestUtilities::UnloadBlueprint(OriginalBlueprint);
-			NewBlueprintUnloader.CloseScope();
-		}
-		else 
-		{
-			// no need to unload the blueprint or any of its dependencies (since 
-			// we didn't muck with it by renaming it) 
-			NewBlueprintUnloader.ClearScope();
-		}
-
-#if WITH_EDITOR
-		// clear undo history to ensure that the transaction buffer isn't 
-		// holding onto any references to the blueprints we want unloaded
-		GEditor->Trans->Reset(NSLOCTEXT("BpAutomation", "RenameCloneTest", "Rename and Clone Test"));
-#endif // #if WITH_EDITOR
-		// make sure the unloaded blueprints are properly flushed (for future tests)
-		CollectGarbage(RF_NoFlags);
-	}
-
-	return !bTestFailed;
-}
