FNames containing asset paths are deprecated. FSoftObjectPath or FTopLevelAssetPath should be used instead.
Fixups for some test files.

diff --git a/Engine/Plugins/Tests/EditorTests/Source/EditorTests/Private/UnrealEd/BlueprintAutomationTests.cpp b/Engine/Plugins/Tests/EditorTests/Source/EditorTests/Private/UnrealEd/BlueprintAutomationTests.cpp
--- a/Engine/Plugins/Tests/EditorTests/Source/EditorTests/Private/UnrealEd/BlueprintAutomationTests.cpp
+++ b/Engine/Plugins/Tests/EditorTests/Source/EditorTests/Private/UnrealEd/BlueprintAutomationTests.cpp
@@ -264,7 +264,7 @@ public:
 		for (auto ObjIter=ObjectList.CreateConstIterator(); ObjIter; ++ObjIter)
 		{
 			const FAssetData& Asset = *ObjIter;
-			FString Filename = Asset.ObjectPath.ToString();
+			FString Filename = Asset.GetObjectPathString();
 			//convert to full paths
 			Filename = FPackageName::LongPackageNameToFilename(Filename);
 			if (FAutomationTestFramework::Get().ShouldTestContent(Filename))
@@ -274,7 +274,7 @@ public:
 				{
 					FString BeautifiedFilename = Asset.AssetName.ToString();
 					OutBeautifiedNames.Add(BeautifiedFilename);
-					OutTestCommands.Add(Asset.ObjectPath.ToString());
+					OutTestCommands.Add(Asset.GetObjectPathString());
 				}
 			}
 		}
diff --git a/Engine/Plugins/Tests/EditorTests/Source/EditorTests/Private/UnrealEd/PlayMapInPIETests.cpp b/Engine/Plugins/Tests/EditorTests/Source/EditorTests/Private/UnrealEd/PlayMapInPIETests.cpp
--- a/Engine/Plugins/Tests/EditorTests/Source/EditorTests/Private/UnrealEd/PlayMapInPIETests.cpp
+++ b/Engine/Plugins/Tests/EditorTests/Source/EditorTests/Private/UnrealEd/PlayMapInPIETests.cpp
@@ -115,7 +115,7 @@ public:
 					{
 						if (!IsDeveloperDirectoryIncluded && MapPackageName.Find(TEXT("/Game/Developers")) == 0) continue;
 
-						FString MapAssetPath = MapAsset.ObjectPath.ToString();
+						FString MapAssetPath = MapAsset.GetObjectPathString();
 
 						OutBeautifiedNames.Add(MapPackageName.RightChop(6).Replace(TEXT("/"), TEXT("."))); // Remove "/Game/" from the name and use dot syntax
 						OutTestCommands.Add(MapAssetPath + TEXT(";") + MapPackageName);
@@ -214,7 +214,7 @@ void OpenMap(const TArray<FString>& Args)
 
 	FAssetRegistryModule& AssetRegistryModule = FModuleManager::LoadModuleChecked<FAssetRegistryModule>(TEXT("AssetRegistry"));
 
-	FAssetData MapAssetData = AssetRegistryModule.Get().GetAssetByObjectPath(*AssetPath);
+	FAssetData MapAssetData = AssetRegistryModule.Get().GetAssetByObjectPath(FSoftObjectPath(AssetPath));
 
 	if (MapAssetData.IsValid())
 	{
diff --git a/Engine/Plugins/Tests/FunctionalTestingEditor/Source/Private/FunctionalTestingEditorModule.cpp b/Engine/Plugins/Tests/FunctionalTestingEditor/Source/Private/FunctionalTestingEditorModule.cpp
--- a/Engine/Plugins/Tests/FunctionalTestingEditor/Source/Private/FunctionalTestingEditorModule.cpp
+++ b/Engine/Plugins/Tests/FunctionalTestingEditor/Source/Private/FunctionalTestingEditorModule.cpp
@@ -41,7 +41,7 @@ void OpenMapAndFocusActor(const TArray<FString>& Args)
 
 	FAssetRegistryModule& AssetRegistryModule = FModuleManager::LoadModuleChecked<FAssetRegistryModule>(TEXT("AssetRegistry"));
 
-	FAssetData MapAssetData = AssetRegistryModule.Get().GetAssetByObjectPath(*AssetPath);
+	FAssetData MapAssetData = AssetRegistryModule.Get().GetAssetByObjectPath(FSoftObjectPath(AssetPath));
 
 	if ( MapAssetData.IsValid() )
 	{
diff --git a/Engine/Source/Editor/UATHelper/UATHelperModule.cpp b/Engine/Source/Editor/UATHelper/UATHelperModule.cpp
--- a/Engine/Source/Editor/UATHelper/UATHelperModule.cpp
+++ b/Engine/Source/Editor/UATHelper/UATHelperModule.cpp
@@ -204,12 +204,12 @@ private:
 				{
 					// Generate qualified asset path and query the registry
 					AssetPath = LongPackageName + TEXT(".") + FPackageName::GetShortName(LongPackageName);
-					FName AssetPathName(*AssetPath, FNAME_Find);
-					if (!AssetPathName.IsNone())
+					FSoftObjectPath SoftAssetPath(AssetPath);
+					if (!SoftAssetPath.IsNull())
 					{
 						static const FName AssetRegistryModuleName(TEXT("AssetRegistry"));
 						FAssetRegistryModule& AssetRegistryModule = FModuleManager::LoadModuleChecked<FAssetRegistryModule>(AssetRegistryModuleName);
-						AssetData = AssetRegistryModule.Get().GetAssetByObjectPath(AssetPathName, true);
+						AssetData = AssetRegistryModule.Get().GetAssetByObjectPath(SoftAssetPath, true);
 					}
 				}
 			}
