FbxAutomation cannot anymore turn off distance field calculation.
#rnx

diff --git a/Engine/Plugins/Tests/EditorTests/Source/EditorTests/Private/UnrealEd/FbxAutomationTests.cpp b/Engine/Plugins/Tests/EditorTests/Source/EditorTests/Private/UnrealEd/FbxAutomationTests.cpp
--- a/Engine/Plugins/Tests/EditorTests/Source/EditorTests/Private/UnrealEd/FbxAutomationTests.cpp
+++ b/Engine/Plugins/Tests/EditorTests/Source/EditorTests/Private/UnrealEd/FbxAutomationTests.cpp
@@ -238,16 +238,6 @@ bool FFbxImportAssetsAutomationTest::RunTest(const FString& Parameters)
 	//Read the fbx options from the .json file and fill the ImportUI
 	FbxAutomationTestsAPI::ReadFbxOptions(FileOptionAndResult, TestPlanArray);
 
-	static const auto CVarDistanceField = IConsoleManager::Get().FindTConsoleVariableDataInt(TEXT("r.GenerateMeshDistanceFields"));
-	int32 OriginalCVarDistanceFieldValue = CVarDistanceField->GetValueOnGameThread();
-	IConsoleVariable* CVarDistanceFieldInterface = IConsoleManager::Get().FindConsoleVariable(TEXT("r.GenerateMeshDistanceFields"));
-	//Avoid building the distance field when we do fbx automation test
-	if (OriginalCVarDistanceFieldValue != 0 && CVarDistanceFieldInterface)
-	{
-		//Hack we change the distance field user console variable to control the build, but we put back the value after the first build
-		CVarDistanceFieldInterface->SetWithCurrentPriority(0);
-	}
-
 	FAssetToolsModule& AssetToolsModule = FModuleManager::LoadModuleChecked<FAssetToolsModule>("AssetTools");
 	bool CurTestSuccessful = (TestPlanArray.Num() > 0);
 	TArray<UObject*> GlobalImportedObjects;
@@ -2023,12 +2013,6 @@ bool FFbxImportAssetsAutomationTest::RunTest(const FString& Parameters)
 		}
 	}
 	
-	//Put back the distance field value
-	if (OriginalCVarDistanceFieldValue != 0 && CVarDistanceFieldInterface)
-	{
-		CVarDistanceFieldInterface->SetWithCurrentPriority(OriginalCVarDistanceFieldValue);
-	}
-
 	return CurTestSuccessful;
 }
 
