Use GetObjectPathString instead of deprecated ObjectPath field.

diff --git a/Engine/Plugins/Tests/EditorTests/Source/EditorTests/Private/EditorUtilityTest.cpp b/Engine/Plugins/Tests/EditorTests/Source/EditorTests/Private/EditorUtilityTest.cpp
--- a/Engine/Plugins/Tests/EditorTests/Source/EditorTests/Private/EditorUtilityTest.cpp
+++ b/Engine/Plugins/Tests/EditorTests/Source/EditorTests/Private/EditorUtilityTest.cpp
@@ -280,7 +280,7 @@ void FBlueprintEditorUtilitiesTest::GetTests(TArray<FString>& OutBeautifiedNames
 	for (auto ObjIter = ObjectList.CreateConstIterator(); ObjIter; ++ObjIter)
 	{
 		const FAssetData& Asset = *ObjIter;
-		FString Filename = Asset.ObjectPath.ToString();
+		FString Filename = Asset.GetObjectPathString();
 
 		// Get the the class this blueprint generates
 		auto AssetTag = Asset.TagsAndValues.FindTag(TEXT("GeneratedClass"));
@@ -305,7 +305,7 @@ void FBlueprintEditorUtilitiesTest::GetTests(TArray<FString>& OutBeautifiedNames
 				FString BeautifiedFilename = Filename + TEXT(".") + Asset.AssetName.ToString();
 
 				OutBeautifiedNames.Add(BeautifiedFilename);
-				OutTestCommands.Add(Asset.ObjectPath.ToString());
+				OutTestCommands.Add(Asset.GetObjectPathString());
 			}
 		}
 	}
