Fix GatherObjectReferencerForDeletion test. We must make sure the CVarUseLegacyGetReferencersForDeletion is set to 0 before running the test. There is some logic when we delete file that turn the cvar to 1. Some automation tests that delete object can trig this situation see Interchange.Import.Scene_Basic test

#rnx

diff --git a/Engine/Plugins/Tests/EditorTests/Source/EditorTests/Private/UnrealEd/ObjectToolsTests.cpp b/Engine/Plugins/Tests/EditorTests/Source/EditorTests/Private/UnrealEd/ObjectToolsTests.cpp
--- a/Engine/Plugins/Tests/EditorTests/Source/EditorTests/Private/UnrealEd/ObjectToolsTests.cpp
+++ b/Engine/Plugins/Tests/EditorTests/Source/EditorTests/Private/UnrealEd/ObjectToolsTests.cpp
@@ -4,6 +4,7 @@
 
 #include "Misc/Paths.h"
 #include "Misc/AutomationTest.h"
+#include "Misc/ScopeExit.h"
 #include "Editor.h"
 #include "Editor/Transactor.h"
 #include "ObjectTools.h"
@@ -206,6 +207,20 @@ bool FObjectToolsTests_GatherObjectReferencersForDeletion::RunTest(const FString
 
 	UObjectToolsTestObject* ObjectInPackage = NewObject<UObjectToolsTestObject>(TempPackage, NAME_None, RF_Transactional);
 
+	//Make sure we do not use the legacy referencers for deletion
+	static IConsoleVariable* CVarUseLegacyGetReferencersForDeletion = IConsoleManager::Get().FindConsoleVariable(TEXT("Editor.UseLegacyGetReferencersForDeletion"));
+	check(CVarUseLegacyGetReferencersForDeletion);
+	bool PreviousCVarUseLegacyGetReferencersForDeletionValue = CVarUseLegacyGetReferencersForDeletion->GetBool();
+	CVarUseLegacyGetReferencersForDeletion->Set(0);
+	//Put back the previous value when the test is done
+	ON_SCOPE_EXIT
+	{
+		if (PreviousCVarUseLegacyGetReferencersForDeletionValue)
+		{
+			CVarUseLegacyGetReferencersForDeletion->Set(1);
+		}
+	};
+
 	// Test direct references to ObjectInPackage
 	{
 		UObjectToolsTestObject* DirectRootReferencer = NewObject<UObjectToolsTestObject>(GetTransientPackage(), NAME_None, RF_Transactional);
