Restructure test to allow PIE as soon as the level is loaded to focus on responsiveness under heavy load and asset priorization during PIE
Nonetheless wait until all assets have finished compiling before completing the test so that we can monitor total cold time

#rnx

diff --git a/Engine/Plugins/Tests/EditorTests/Source/EditorTests/Private/UnrealEd/EditorIterationTests.cpp b/Engine/Plugins/Tests/EditorTests/Source/EditorTests/Private/UnrealEd/EditorIterationTests.cpp
--- a/Engine/Plugins/Tests/EditorTests/Source/EditorTests/Private/UnrealEd/EditorIterationTests.cpp
+++ b/Engine/Plugins/Tests/EditorTests/Source/EditorTests/Private/UnrealEd/EditorIterationTests.cpp
@@ -47,10 +47,6 @@ bool FIterationOpenAssets::RunTest(const FString& LongAssetPath)
 	// Issue Load request
 	AddCommand(new FOpenEditorForAssetCommand(LongAssetPath));
 
-	// Wait on all async asset processing
-	AddCommand(new FFunctionLatentCommand([] {
-		return FAssetCompilingManager::Get().GetNumRemainingAssets() == 0;
-	}));
 	AddCommand(new FFunctionLatentCommand([LongAssetPath] {
 		TRACE_BOOKMARK(TEXT("LoadAssetComplete - %s"), *LongAssetPath);
 		TRACE_BOOKMARK(TEXT("PIE - %s"), *LongAssetPath);
@@ -73,5 +69,15 @@ bool FIterationOpenAssets::RunTest(const FString& LongAssetPath)
 		return true;
 		}));
 
+	// Wait on all async asset processing
+	AddCommand(new FFunctionLatentCommand([] {
+		return FAssetCompilingManager::Get().GetNumRemainingAssets() == 0;
+		}));
+
+	AddCommand(new FFunctionLatentCommand([LongAssetPath] {
+		TRACE_BOOKMARK(TEXT("AssetCompilationComplete - %s"), *LongAssetPath);
+		return true;
+		}));
+
 	return true;
 }
