Fix GatherObjectReferencersForDeletion allowing deletion of package even if internal objects have RF_Standalone flag on them

#rnx

diff --git a/Engine/Plugins/Tests/EditorTests/Source/EditorTests/Private/UnrealEd/ObjectToolsTests.cpp b/Engine/Plugins/Tests/EditorTests/Source/EditorTests/Private/UnrealEd/ObjectToolsTests.cpp
--- a/Engine/Plugins/Tests/EditorTests/Source/EditorTests/Private/UnrealEd/ObjectToolsTests.cpp
+++ b/Engine/Plugins/Tests/EditorTests/Source/EditorTests/Private/UnrealEd/ObjectToolsTests.cpp
@@ -277,6 +277,17 @@ bool FObjectToolsTests_GatherObjectReferencersForDeletion::RunTest(const FString
 		TestGatherObjectReferencersForDeletion(TempPackage, TEXT("GatherObjectReferencersForDeletion shouldn't detect any reference on any objects inside TempPackage"), false, false);
 	}
 
+	// Test object inside the package with RF_Standalone flag
+	{
+		UObjectToolsTestObject* StandaloneObject = NewObject<UObjectToolsTestObject>(TempPackage, NAME_None, RF_Transactional | RF_Standalone);
+
+		TestGatherObjectReferencersForDeletion(TempPackage, TEXT("GatherObjectReferencersForDeletion should detect internal reference inside TempPackage"), true, false);
+
+		StandaloneObject->ClearFlags(RF_Standalone);
+
+		TestGatherObjectReferencersForDeletion(TempPackage, TEXT("GatherObjectReferencersForDeletion shouldn't detect any reference on any objects inside TempPackage"), false, false);
+	}
+
 	return true;
 }
 
diff --git a/Engine/Source/Editor/UnrealEd/Private/ObjectTools.cpp b/Engine/Source/Editor/UnrealEd/Private/ObjectTools.cpp
--- a/Engine/Source/Editor/UnrealEd/Private/ObjectTools.cpp
+++ b/Engine/Source/Editor/UnrealEd/Private/ObjectTools.cpp
@@ -264,7 +264,9 @@ namespace ObjectTools
 			// Get the cluster of objects that are going to be deleted
 			TArray<UObject*> ObjectsToDelete;
 			GetObjectsWithOuter(InObject, ObjectsToDelete);
-			
+
+			bool bIsReferencedInternally = false;
+
 			TSet<UObject*> InternalReferences;
 			// The old behavior of GatherObjectReferencersForDeletion will find anything that prevents 
 			// InObject from being garbage collected, including internal sub objects.
@@ -276,6 +278,7 @@ namespace ObjectTools
 				{
 					InternalReferences.Add(ObjectToDelete);
 					bOutIsReferenced = true;
+					bIsReferencedInternally = true;
 				}
 			}
 			
@@ -360,11 +363,11 @@ namespace ObjectTools
 				if (Roots.Contains(Transactor))
 				{
 					bOutIsReferencedInMemoryByUndo = true;
-					bOutIsReferenced = Roots.Num() > 1;
+					bOutIsReferenced = Roots.Num() > 1 || bIsReferencedInternally;
 				}
 				else
 				{
-					bOutIsReferenced = Roots.Num() > 0;
+					bOutIsReferenced = Roots.Num() > 0 || bIsReferencedInternally;
 				}
 			}
 
