Added check GIsAutomationTesting to avoid [Owner] and [Description] log in PIE

diff --git a/Engine/Plugins/Tests/EditorTests/Source/EditorTests/Private/EditorUtilityTest.cpp b/Engine/Plugins/Tests/EditorTests/Source/EditorTests/Private/EditorUtilityTest.cpp
--- a/Engine/Plugins/Tests/EditorTests/Source/EditorTests/Private/EditorUtilityTest.cpp
+++ b/Engine/Plugins/Tests/EditorTests/Source/EditorTests/Private/EditorUtilityTest.cpp
@@ -75,7 +75,7 @@ void UEditorUtilityTest::Run()
 			CurrentTest->SetEditorUtility(this);
 		}
 
-		if (FAutomationTestFramework::NeedLogBPTestMetadata())
+		if (FAutomationTestFramework::NeedLogBPTestMetadata() && GIsAutomationTesting)
 		{
 			AddInfo(FString::Printf(TEXT("[Owner] %s"), *Owner));
 			AddInfo(FString::Printf(TEXT("[Description] %s"), *Description));
diff --git a/Engine/Source/Developer/FunctionalTesting/Private/FunctionalTest.cpp b/Engine/Source/Developer/FunctionalTesting/Private/FunctionalTest.cpp
--- a/Engine/Source/Developer/FunctionalTesting/Private/FunctionalTest.cpp
+++ b/Engine/Source/Developer/FunctionalTesting/Private/FunctionalTest.cpp
@@ -295,7 +295,7 @@ void AFunctionalTest::StartTest()
 	StartFrame = GFrameNumber;
 	StartTime = (float)GetWorld()->GetTimeSeconds();
 
-	if (FAutomationTestFramework::NeedLogBPTestMetadata())
+	if (FAutomationTestFramework::NeedLogBPTestMetadata() && GIsAutomationTesting)
 	{
 		AddInfo(FString::Printf(TEXT("[Owner] %s"), *Author));
 		AddInfo(FString::Printf(TEXT("[Description] %s"), *Description));
