Increase the delay for 'Auto Reimport' tests to prevent incorrect reporting on slower machines.

#rnx

- On internal CIS mac machines these tests were sometimes failing, but this was never seen on local machines.
- Ran many tests, with very verbose logging added and the conclusion was that the file events were being found but some times they arrived too late and the test had already "timed out" or the file cache watcher has interpreted the results incorrectly because not all of the important events had arrived.
- Increasing the time out seems to work. This is not a good fix but reworking both the watcher code, the file cache code and the tests is out of scope at the moment.

diff --git a/Engine/Plugins/Tests/EditorTests/Source/EditorTests/Private/UnrealEd/AutoReimportTests.cpp b/Engine/Plugins/Tests/EditorTests/Source/EditorTests/Private/UnrealEd/AutoReimportTests.cpp
--- a/Engine/Plugins/Tests/EditorTests/Source/EditorTests/Private/UnrealEd/AutoReimportTests.cpp
+++ b/Engine/Plugins/Tests/EditorTests/Source/EditorTests/Private/UnrealEd/AutoReimportTests.cpp
@@ -157,7 +157,7 @@ bool CopyTestFiles(const FAutoReimportTestPayload& Test, const TArray<FSrcDstFil
 }
 
 const int32 LocalMaxRetryCounter = 5;
-const float LocalRetryDelay = 1.0f;
+const float LocalRetryDelay = 5.0f;
 
 void RetryLatentCommand(TFunction<bool(const bool)>&& LatentCommandFunction, const float RetryDelay, int32& MaxRetryCounter)
 {
