New Editor Test:
In a test map, spawn an actor, then select it, and then check to see if the spawned actor is selected.

diff --git a/Engine/Plugins/Tests/EditorTests/Source/EditorTests/Private/UnrealEd/EditorTestSimpleSelection.cpp b/Engine/Plugins/Tests/EditorTests/Source/EditorTests/Private/UnrealEd/EditorTestSimpleSelection.cpp
new file mode 100644
--- /dev/null
+++ b/Engine/Plugins/Tests/EditorTests/Source/EditorTests/Private/UnrealEd/EditorTestSimpleSelection.cpp
@@ -0,0 +1,85 @@
+// Copyright Epic Games, Inc. All Rights Reserved.
+#include "CoreMinimal.h"
+#include "Misc/AutomationTest.h"
+#include "Editor/UnrealEd/Public/Editor.h" 
+#include "Editor/UnrealEd/Public/Selection.h"
+#include "FileHelpers.h"
+// Asynchronous function to select the input actor
+DEFINE_LATENT_AUTOMATION_COMMAND_ONE_PARAMETER(FBSelectActor, AActor*, TestAsset);
+bool FBSelectActor::Update()
+{
+	constexpr bool bInSelected = true;
+	constexpr bool bNotify = true;
+
+	GEditor->SelectActor(TestAsset, bInSelected, bNotify);
+	return true;
+}
+// Asynchronous function that checks whether the input actor is selected
+DEFINE_LATENT_AUTOMATION_COMMAND_THREE_PARAMETER(FBVerifySimpleSelection, FAutomationTestBase*, Test, AActor*, TestAsset, FString, MapName);
+bool FBVerifySimpleSelection::Update()
+{
+	if (TestAsset == nullptr)
+	{
+		Test->AddError(TEXT("The target actor is a null pointer."));
+	}
+	else
+	{
+		bool bIsSelected = GEditor->GetSelectedActors()->IsSelected(TestAsset);
+		FString ActorName = TestAsset->GetName();
+		// Check if the actor is selected
+		if (!bIsSelected)
+		{
+			Test->AddError(FString::Printf(TEXT("The target actor: '%s' is NOT selected."), *ActorName));
+		}
+		else
+		{
+			UE_LOG(LogTemp, Display, TEXT("The target actor: '%s' is successfully selected."), *ActorName);
+		}
+	}
+	// Always attempt to load the map, regardless of whether the selection was successful.
+	if (!FEditorFileUtils::LoadMap(MapName, false, false))
+	{
+		Test->AddError(FString::Printf(TEXT("Failed to load the previous map: '%s'"), *MapName));
+	}
+	return true;
+}
+IMPLEMENT_SIMPLE_AUTOMATION_TEST(FAssetSimpleSelectionTest, "Editor.Selection.SimpleSelection", EAutomationTestFlags::EditorContext | EAutomationTestFlags::ProductFilter)
+bool FAssetSimpleSelectionTest::RunTest(const FString& Parameters)
+{
+	FString PreviousLevel = GEditor->GetEditorWorldContext().World()->GetOutermost()->GetName();
+	if (PreviousLevel.Find("/Temp/Untitled") != INDEX_NONE)
+	{
+		// Must save the current map before running this test
+		AddError(TEXT("Please save your current level before running this test."));
+		return false;
+	}
+	// Load the map for testing
+	FString TestLevelPath = "/Game/Tests/Selection/SimpleSelectionTestMap";
+	if (!FEditorFileUtils::LoadMap(TestLevelPath))
+	{
+		AddError(FString::Printf(TEXT("Failed to load the test map: '%s'"), *TestLevelPath));
+		return false;
+	}
+	// Get the world pointer from the editor
+	UWorld* TestWorld = GEditor->GetEditorWorldContext().World();
+	if (TestWorld == nullptr)
+	{
+		AddError(TEXT("Failed to get the test world from the editor."));
+		return false;
+	}
+	// Spawn the test asset actor
+	AActor* TestAsset = TestWorld->SpawnActor<AActor>(FVector::ZeroVector, FRotator::ZeroRotator);
+	if (TestAsset == nullptr)
+	{
+		AddError(TEXT("Failed to spawn the test asset."));
+		// Always attempt to load the map, regardless of whether the selection was successful.
+		if (!FEditorFileUtils::LoadMap(PreviousLevel, false, false))
+		{
+			AddError(FString::Printf(TEXT("Failed to load the previous map: '%s'"), *PreviousLevel));
+		}
+		return false;
+	}
+	ADD_LATENT_AUTOMATION_COMMAND(FBSelectActor(TestAsset));
+	ADD_LATENT_AUTOMATION_COMMAND(FBVerifySimpleSelection(this, TestAsset, PreviousLevel));
+	return true;
+}
\ No newline at end of file
