Editor Workflows Selection Tests
EditorTestActorSelection.cpp: spawns an empty actor, selects it, and then verifies the selection set.
EditorTestHierarchicalISMSelection.cpp: finds a pre-existing actor in the test map, selects the ISM component of that actor, then verifies the selection set.
BP_BadPerfSelectingISM.uasset: test actor used in EditorTestHierarchicalISMSelection.cpp
HierarchicalISMSelectionTestMap.umap: test map used in both EditorTestHierarchicalISMSelection.cpp and EditorTestActorSelection.cpp

diff --git a/Engine/Plugins/Tests/EditorTests/Source/EditorTests/Private/UnrealEd/EditorTestActorSelection.cpp b/Engine/Plugins/Tests/EditorTests/Source/EditorTests/Private/UnrealEd/EditorTestActorSelection.cpp
new file mode 100644
--- /dev/null
+++ b/Engine/Plugins/Tests/EditorTests/Source/EditorTests/Private/UnrealEd/EditorTestActorSelection.cpp
@@ -0,0 +1,106 @@
+// Copyright Epic Games, Inc. All Rights Reserved.
+
+#include "CoreMinimal.h"
+#include "Misc/AutomationTest.h"
+#include "Tests/AutomationEditorCommon.h"
+#include "Tests/AutomationCommon.h"
+#include "Editor/UnrealEd/Public/Editor.h" 
+#include "Editor/UnrealEd/Public/Selection.h"
+#include "FileHelpers.h"
+
+// Asynchronous function to select the input actor
+DEFINE_LATENT_AUTOMATION_COMMAND_ONE_PARAMETER(FBSelectActor, AActor**, TestAsset);
+bool FBSelectActor::Update()
+{
+	constexpr bool bInSelected = true;
+	constexpr bool bNotify = true;
+
+	GEditor->SelectActor(*TestAsset, bInSelected, bNotify);
+
+	return true;
+}
+
+// Asynchronous function that checks whether the input actor is selected or not
+DEFINE_LATENT_AUTOMATION_COMMAND_TWO_PARAMETER(FBVerifyActorSelection, FAutomationTestBase*, Test, AActor**, TestAsset);
+bool FBVerifyActorSelection::Update()
+{
+	if (*TestAsset != nullptr)
+	{
+		bool bIsSelected = GEditor->GetSelectedActors()->IsSelected(*TestAsset);
+		FString ActorName = (*TestAsset)->GetName();
+
+		// Check if the actor is selected
+		if (bIsSelected)
+		{
+			UE_LOG(LogTemp, Display, TEXT("The target actor: '%s' is successfully selected."), *ActorName);
+		}
+		else
+		{
+			Test->AddError(FString::Printf(TEXT("The target actor: '%s' is NOT selected."), *ActorName));
+		}
+	}
+	else
+	{
+		Test->AddError(TEXT("The target actor is a null pointer."));
+	}
+
+	return true;
+}
+
+IMPLEMENT_SIMPLE_AUTOMATION_TEST(FAssetActorTest, "Editor.Workflows.Selection.Actor", EAutomationTestFlags::EditorContext | EAutomationTestFlags::ProductFilter)
+bool FAssetActorTest::RunTest(const FString& Parameters)
+{
+	UWorld** TestWorld = (UWorld**)FMemory::Malloc(sizeof(UWorld**));
+	AActor** TestAsset = (AActor**)FMemory::Malloc(sizeof(AActor**));
+
+	const FString EditorTestName = TEXT("Editor.Workflows.Selection.Actor");
+	const TCHAR* TestMapPath = TEXT("/Game/Tests/Selection/HierarchicalISMSelectionTestMap");
+
+	// Setup
+	AddCommand(new FCloseAllAssetEditorsCommand());
+	AddCommand(new FFunctionLatentCommand([EditorTestName] {
+		TRACE_BOOKMARK(*(EditorTestName + TEXT(" ProfileBegin")));
+
+		return true;
+		}));
+
+	// Issue Load request
+	// Spawn a latent command to wait for map loading to complete
+	AddCommand(new FOpenEditorForAssetCommand(TestMapPath));
+	AddCommand(new FFunctionLatentCommand([TestWorld]() {
+		*TestWorld = GEditor->GetEditorWorldContext().World();
+
+		return *TestWorld != nullptr;
+		}));
+
+	AddCommand(new FWaitForEngineFramesCommand(20));
+
+	AddCommand(new FFunctionLatentCommand([TestAsset, TestWorld, this, EditorTestName] {
+		// Spawn the test asset actor
+		*TestAsset = (*TestWorld)->SpawnActor<AActor>(FVector::ZeroVector, FRotator::ZeroRotator);
+		if (TestAsset == nullptr)
+		{
+			AddError(EditorTestName + TEXT(": Failed to spawn the test asset."));
+		}
+
+		return true;
+		}));
+
+	ADD_LATENT_AUTOMATION_COMMAND(FBSelectActor(TestAsset));
+
+	ADD_LATENT_AUTOMATION_COMMAND(FBVerifyActorSelection(this, TestAsset));
+
+	AddCommand(new FWaitForEngineFramesCommand(20));
+
+	// Teardown
+	AddCommand(new FFunctionLatentCommand([TestAsset, TestWorld, EditorTestName] {
+		FMemory::Free(TestWorld);
+		FMemory::Free(TestAsset);
+
+		TRACE_BOOKMARK(*(EditorTestName + TEXT(" ProfileEnd")));
+
+		return true;
+		}));
+
+	return true;
+}
\ No newline at end of file
diff --git a/Engine/Plugins/Tests/EditorTests/Source/EditorTests/Private/UnrealEd/EditorTestHierarchicalISMSelection.cpp b/Engine/Plugins/Tests/EditorTests/Source/EditorTests/Private/UnrealEd/EditorTestHierarchicalISMSelection.cpp
new file mode 100644
--- /dev/null
+++ b/Engine/Plugins/Tests/EditorTests/Source/EditorTests/Private/UnrealEd/EditorTestHierarchicalISMSelection.cpp
@@ -0,0 +1,121 @@
+// Copyright Epic Games, Inc. All Rights Reserved.
+
+#include "CoreMinimal.h"
+#include "Misc/AutomationTest.h"
+#include "Tests/AutomationEditorCommon.h"
+#include "Tests/AutomationCommon.h"
+#include "Editor/UnrealEd/Public/Editor.h" 
+#include "Editor/UnrealEd/Public/Selection.h" 
+#include "Components/InstancedStaticMeshComponent.h"
+#include "FileHelpers.h"
+#include "EngineUtils.h"
+
+// Asynchronous function to select the ISM component
+DEFINE_LATENT_AUTOMATION_COMMAND_ONE_PARAMETER(FBSelectISMComponent, UInstancedStaticMeshComponent**, ISMComponent);
+bool FBSelectISMComponent::Update()
+{
+	constexpr bool bInSelected = true;
+	constexpr bool bNotify = true;
+	constexpr bool bSelectEvenIfHidden = true;
+
+	GEditor->SelectComponent(*ISMComponent, bInSelected, bNotify, bSelectEvenIfHidden);
+
+	return true;
+}
+
+// Asynchronous function that checks whether the input ISM component is selected or not
+DEFINE_LATENT_AUTOMATION_COMMAND_TWO_PARAMETER(FBVerifyISMComponentSelection, FAutomationTestBase*, Test, UInstancedStaticMeshComponent**, ISMComponent);
+bool FBVerifyISMComponentSelection::Update()
+{
+	if (*ISMComponent != nullptr)
+	{
+		bool bIsSelected = GEditor->GetSelectedComponents()->IsSelected(*ISMComponent);
+		FString ComponentName = (*ISMComponent)->GetName();
+
+		// Check if the ISM component is selected
+		if (bIsSelected)
+		{
+			UE_LOG(LogTemp, Display, TEXT("The target ISM Component: '%s' is successfully selected."), *ComponentName);
+		}
+		else
+		{
+			Test->AddError(FString::Printf(TEXT("The target ISM Component: '%s' is NOT selected."), *ComponentName));
+		}
+	}
+	else
+	{
+		Test->AddError(TEXT("The target ISM component is a null pointer."));
+	}
+
+	return true;
+}
+
+IMPLEMENT_SIMPLE_AUTOMATION_TEST(FAssetHierarchicalISMSelectionTest, "Editor.Workflows.Selection.HierarchicalISM", EAutomationTestFlags::EditorContext | EAutomationTestFlags::ProductFilter)
+bool FAssetHierarchicalISMSelectionTest::RunTest(const FString& Parameters)
+{
+	UWorld** TestWorld = (UWorld**)FMemory::Malloc(sizeof(UWorld**));
+	UInstancedStaticMeshComponent** ISMComponent = (UInstancedStaticMeshComponent**)FMemory::Malloc(sizeof(UInstancedStaticMeshComponent**));
+
+	const FString EditorTestName = TEXT("Editor.Workflows.Selection.HierarchicalISM");
+	const TCHAR* TestMapPath = TEXT("/Game/Tests/Selection/HierarchicalISMSelectionTestMap");
+
+	// Setup
+	AddCommand(new FCloseAllAssetEditorsCommand());
+	AddCommand(new FFunctionLatentCommand([EditorTestName] {
+		TRACE_BOOKMARK(*(EditorTestName + TEXT(" ProfileBegin")));
+
+		return true;
+		}));
+
+	// Issue Load request
+	// Spawn a latent command to wait for map loading to complete
+	AddCommand(new FOpenEditorForAssetCommand(TestMapPath));
+	AddCommand(new FFunctionLatentCommand([TestWorld]() {
+		*TestWorld = GEditor->GetEditorWorldContext().World();
+
+		return *TestWorld != nullptr;
+		}));
+
+	AddCommand(new FWaitForEngineFramesCommand(20));
+
+	AddCommand(new FFunctionLatentCommand([ISMComponent, TestWorld, this, EditorTestName] {
+		// Get the target actor in the test world
+		AActor* TestAsset = nullptr;
+		for (TActorIterator<AActor> ActorIt(*TestWorld); ActorIt; ++ActorIt)
+		{
+			if (ActorIt->GetActorLabel() == "BP_BadPerfSelectingISM")
+			{
+				TestAsset = *ActorIt;
+				break;
+			}
+		}
+		if (TestAsset != nullptr)
+		{
+			*ISMComponent = TestAsset->FindComponentByClass<UInstancedStaticMeshComponent>();
+		}
+		else
+		{
+			AddError(EditorTestName + TEXT(": Test asset not found in test map."));
+		}
+
+		return true;
+		}));
+
+	ADD_LATENT_AUTOMATION_COMMAND(FBSelectISMComponent(ISMComponent));
+
+	ADD_LATENT_AUTOMATION_COMMAND(FBVerifyISMComponentSelection(this, ISMComponent));
+
+	AddCommand(new FWaitForEngineFramesCommand(20));
+
+	// Teardown
+	AddCommand(new FFunctionLatentCommand([ISMComponent, TestWorld, EditorTestName] {
+		FMemory::Free(TestWorld);
+		FMemory::Free(ISMComponent);
+
+		TRACE_BOOKMARK(*(EditorTestName + TEXT(" ProfileEnd")));
+
+		return true;
+		}));
+
+	return true;
+}
\ No newline at end of file
diff --git a/Engine/Plugins/Tests/EditorTests/Source/EditorTests/Private/UnrealEd/EditorTestSimpleSelection.cpp b/Engine/Plugins/Tests/EditorTests/Source/EditorTests/Private/UnrealEd/EditorTestSimpleSelection.cpp
deleted file mode 100644
