List Editor Utility tests from plugin

#rnx

diff --git a/Engine/Plugins/Tests/EditorTests/Source/EditorTests/Private/EditorUtilityTest.cpp b/Engine/Plugins/Tests/EditorTests/Source/EditorTests/Private/EditorUtilityTest.cpp
--- a/Engine/Plugins/Tests/EditorTests/Source/EditorTests/Private/EditorUtilityTest.cpp
+++ b/Engine/Plugins/Tests/EditorTests/Source/EditorTests/Private/EditorUtilityTest.cpp
@@ -302,17 +302,16 @@ void FBlueprintEditorUtilitiesTest::GetTests(TArray<FString>& OutBeautifiedNames
 			}
 
 			// Exclude Developer folder
-			if (Filename.Find(TEXT("/Game/")) == 0)
-			{
-				if (!IsDeveloperDirectoryIncluded && Filename.Find(TEXT("/Game/Developers")) == 0) continue;
+			if (!IsDeveloperDirectoryIncluded && Filename.Find(TEXT("/Game/Developers")) == 0) continue;
 
-				// Extract folder org
-				Filename = FPaths::GetPath(Filename.RightChop(6)).Replace(TEXT("/"), TEXT("."));
-				FString BeautifiedFilename = Filename + TEXT(".") + Asset.AssetName.ToString();
+			// Extract folder org
+			// Trim /Game/ otherwise use the plugin name as category 
+			int32 TrimLength = Filename.Find(TEXT("/Game/")) == 0 ? 6 : 1;
+			Filename = FPaths::GetPath(Filename.RightChop(TrimLength)).Replace(TEXT("/"), TEXT("."));
+			FString BeautifiedFilename = Filename + TEXT(".") + Asset.AssetName.ToString();
 
-				OutBeautifiedNames.Add(BeautifiedFilename);
-				OutTestCommands.Add(Asset.GetObjectPathString());
-			}
+			OutBeautifiedNames.Add(BeautifiedFilename);
+			OutTestCommands.Add(Asset.GetObjectPathString());
 		}
 	}
 }
