Fix legacy fbx automation tests, which do not play well with interchange conversion.

#rnx

diff --git a/Engine/Plugins/Interchange/Runtime/Source/Import/Private/Fbx/InterchangeFbxTranslator.cpp b/Engine/Plugins/Interchange/Runtime/Source/Import/Private/Fbx/InterchangeFbxTranslator.cpp
--- a/Engine/Plugins/Interchange/Runtime/Source/Import/Private/Fbx/InterchangeFbxTranslator.cpp
+++ b/Engine/Plugins/Interchange/Runtime/Source/Import/Private/Fbx/InterchangeFbxTranslator.cpp
@@ -71,7 +71,7 @@ EInterchangeTranslatorAssetType UInterchangeFbxTranslator::GetSupportedAssetType
 
 TArray<FString> UInterchangeFbxTranslator::GetSupportedFormats() const
 {
-	if (GInterchangeEnableFBXImport || GIsAutomationTesting)
+	if (GInterchangeEnableFBXImport)
 	{
 		TArray<FString> Formats{ TEXT("fbx;Filmbox") };
 		return Formats;
diff --git a/Engine/Plugins/Tests/EditorTests/Source/EditorTests/Private/UnrealEd/FbxAutomationTests.cpp b/Engine/Plugins/Tests/EditorTests/Source/EditorTests/Private/UnrealEd/FbxAutomationTests.cpp
--- a/Engine/Plugins/Tests/EditorTests/Source/EditorTests/Private/UnrealEd/FbxAutomationTests.cpp
+++ b/Engine/Plugins/Tests/EditorTests/Source/EditorTests/Private/UnrealEd/FbxAutomationTests.cpp
@@ -1,48 +1,46 @@
 // Copyright Epic Games, Inc. All Rights Reserved.
 
-#include "CoreMinimal.h"
-#include "Engine/SkinnedAssetCommon.h"
-#include "HAL/FileManager.h"
-#include "Misc/Paths.h"
-#include "Misc/ConfigCacheIni.h"
-#include "Misc/AutomationTest.h"
-#include "Modules/ModuleManager.h"
+#include "Animation/AnimCurveTypes.h"
+#include "Animation/AnimSequence.h"
+#include "Animation/DebugSkelMeshComponent.h"
 #include "Animation/Skeleton.h"
-#include "Engine/SkeletalMesh.h"
-#include "Engine/StaticMesh.h"
-#include "Materials/Material.h"
-#include "Materials/MaterialInstanceConstant.h"
 #include "AssetRegistry/AssetData.h"
-#include "IAssetTools.h"
+#include "AssetRegistry/AssetRegistryModule.h"
 #include "AssetToolsModule.h"
+#include "CoreMinimal.h"
+#include "Editor/EditorEngine.h"
+#include "Editor/Transactor.h"
 #include "EditorReimportHandler.h"
-#include "Factories/FbxFactory.h"
-#include "Factories/ReimportFbxSkeletalMeshFactory.h"
-#include "Factories/ReimportFbxStaticMeshFactory.h"
+#include "Engine/SkeletalMesh.h"
+#include "Engine/SkinnedAssetCommon.h"
+#include "Engine/StaticMesh.h"
 #include "Factories/FbxAnimSequenceImportData.h"
+#include "Factories/FbxFactory.h"
+#include "Factories/FbxImportUI.h"
 #include "Factories/FbxSkeletalMeshImportData.h"
 #include "Factories/FbxStaticMeshImportData.h"
 #include "Factories/FbxTextureImportData.h"
-#include "Factories/FbxImportUI.h"
+#include "Factories/ReimportFbxSkeletalMeshFactory.h"
+#include "Factories/ReimportFbxStaticMeshFactory.h"
+#include "FbxMeshUtils.h"
+#include "HAL/FileManager.h"
+#include "IAssetTools.h"
+#include "InterchangeHelper.h"
+#include "Materials/Material.h"
+#include "Materials/MaterialInstanceConstant.h"
+#include "Misc/Paths.h"
+#include "Misc/ConfigCacheIni.h"
+#include "Misc/AutomationTest.h"
+#include "ObjectTools.h"
+#include "Modules/ModuleManager.h"
+#include "Rendering/SkeletalMeshModel.h"
 #include "Rendering/SkeletalMeshRenderData.h"
 #include "SkinWeightsUtilities.h"
-#include "Animation/DebugSkelMeshComponent.h"
-#include "Rendering/SkeletalMeshModel.h"
-#include "Misc/Paths.h"
-#include "UObject/SavePackage.h"
-
-#include "Animation/AnimCurveTypes.h"
-#include "Animation/AnimSequence.h"
-
-#include "AssetRegistry/AssetRegistryModule.h"
-#include "ObjectTools.h"
 #include "StaticMeshResources.h"
-
-#include "FbxMeshUtils.h"
+#include "Templates/Function.h"
 #include "Tests/FbxAutomationCommon.h"
+#include "UObject/SavePackage.h"
 
-#include "Editor/Transactor.h"
-#include "Editor/EditorEngine.h"
 extern UNREALED_API UEditorEngine* GEditor;
 
 //////////////////////////////////////////////////////////////////////////
@@ -217,6 +215,15 @@
 
 	CollectGarbage(GARBAGE_COLLECTION_KEEPFLAGS);
 
+	static const auto CVarInterchangeFbx = IConsoleManager::Get().FindConsoleVariable(TEXT("Interchange.FeatureFlags.Import.FBX"));
+	bool IsInterchangeFbxEnabled = CVarInterchangeFbx->GetBool();
+	UE::Interchange::FScopedLambda IsInterchangeEnabledGuard([&IsInterchangeFbxEnabled]()
+		{ 
+			CVarInterchangeFbx->Set(IsInterchangeFbxEnabled, ECVF_SetByConsole);
+		});
+	//Make sure interchange is disabled for fbx
+	CVarInterchangeFbx->Set(false, ECVF_SetByConsole);
+
 	FString PackagePath;
 	check(GConfig);
 	GConfig->GetString(TEXT("AutomationTesting.FbxImport"), TEXT("FbxImportTestPackagePath"), PackagePath, GEngineIni);
diff --git a/Engine/Plugins/Tests/InterchangeTests/Source/InterchangeTests/Private/InterchangeImportTest.cpp b/Engine/Plugins/Tests/InterchangeTests/Source/InterchangeTests/Private/InterchangeImportTest.cpp
--- a/Engine/Plugins/Tests/InterchangeTests/Source/InterchangeTests/Private/InterchangeImportTest.cpp
+++ b/Engine/Plugins/Tests/InterchangeTests/Source/InterchangeTests/Private/InterchangeImportTest.cpp
@@ -8,6 +8,7 @@
 #include "HAL/FileManager.h"
 #include "Misc/AutomationTest.h"
 #include "InterchangeDispatcher.h"
+#include "InterchangeHelper.h"
 #include "InterchangeImportTestData.h"
 #include "InterchangeImportTestPlan.h"
 #include "InterchangeImportTestStepBase.h"
@@ -98,10 +99,17 @@ bool FInterchangeImportTest::RunTest(const FString& Path)
 	return true;
 #endif
 
-	// Determine the test plan assets within the given path which will be run in parallel
+	static const auto CVarInterchangeFbx = IConsoleManager::Get().FindConsoleVariable(TEXT("Interchange.FeatureFlags.Import.FBX"));
+	bool IsInterchangeFbxEnabled = CVarInterchangeFbx->GetBool();
+	UE::Interchange::FScopedLambda IsInterchangeEnabledGuard([&IsInterchangeFbxEnabled]()
+		{
+			CVarInterchangeFbx->Set(IsInterchangeFbxEnabled, ECVF_SetByConsole);
+		});
+	//Make sure interchange is enabled for fbx
+	CVarInterchangeFbx->Set(true, ECVF_SetByConsole);
 
+	// Determine the test plan assets within the given path which will be run in parallel
 	TArray<FInterchangeImportTestData> TestPlans;
-
 	if (Path.Contains(TEXT(".")))
 	{
 		// Run test on a single TestPlan asset
diff --git a/Engine/Source/Runtime/Interchange/Core/Public/InterchangeHelper.h b/Engine/Source/Runtime/Interchange/Core/Public/InterchangeHelper.h
--- a/Engine/Source/Runtime/Interchange/Core/Public/InterchangeHelper.h
+++ b/Engine/Source/Runtime/Interchange/Core/Public/InterchangeHelper.h
@@ -12,4 +12,23 @@ namespace UE::Interchange
 	 * #todo_interchange: Replace with SanitizeObjectName from InterchangeManager.cpp
 	 */
 	INTERCHANGECORE_API FString MakeName(const FString& InName, bool bIsJoint = false);
+
+	class FScopedLambda
+	{
+	public:
+		FScopedLambda(TFunction<void(void)>&& Lambda)
+			: _Lambda(Lambda)
+		{}
+
+		~FScopedLambda()
+		{
+			if (_Lambda)
+			{
+				_Lambda();
+			}
+		}
+
+	private:
+		TFunction<void(void)> _Lambda;
+	};
 };
\ No newline at end of file
