Add test for editing of the UWidget::Navigation subobject
#rnx

diff --git a/Engine/Plugins/Tests/EditorTests/Source/EditorTests/EditorTests.Build.cs b/Engine/Plugins/Tests/EditorTests/Source/EditorTests/EditorTests.Build.cs
--- a/Engine/Plugins/Tests/EditorTests/Source/EditorTests/EditorTests.Build.cs
+++ b/Engine/Plugins/Tests/EditorTests/Source/EditorTests/EditorTests.Build.cs
@@ -53,7 +53,8 @@ public class EditorTests : ModuleRules
                 "RawMesh",
 				"AutomationController",
 				"Blutility",
-				//"SubobjectDataInterface",
+				"UMGEditor",
+				"UMG"
 			}
 		);
 	}
diff --git a/Engine/Plugins/Tests/EditorTests/Source/EditorTests/Private/UnrealEd/EditorTestsUtilityLibrary.cpp b/Engine/Plugins/Tests/EditorTests/Source/EditorTests/Private/UnrealEd/EditorTestsUtilityLibrary.cpp
--- a/Engine/Plugins/Tests/EditorTests/Source/EditorTests/Private/UnrealEd/EditorTestsUtilityLibrary.cpp
+++ b/Engine/Plugins/Tests/EditorTests/Source/EditorTests/Private/UnrealEd/EditorTestsUtilityLibrary.cpp
@@ -20,6 +20,9 @@
 #include "Algo/Transform.h"
 #include "Materials/Material.h"
 #include "TextureCompiler.h"
+#include "WidgetBlueprint.h"
+#include "Blueprint/WidgetTree.h"
+#include "Components/Widget.h"
 
 void UEditorTestsUtilityLibrary::BakeMaterialsForComponent(UStaticMeshComponent* InStaticMeshComponent, const UMaterialOptions* MaterialOptions, const UMaterialMergeOptions* MaterialMergeOptions)
 {
@@ -151,3 +154,93 @@ void UEditorTestsUtilityLibrary::MergeStaticMeshComponents(TArray<UStaticMeshCom
 	}
 }
 
+UWidget* UEditorTestsUtilityLibrary::GetChildEditorWidgetByName(UWidgetBlueprint* WidgetBlueprint, FString Name)
+{
+	if (!ensure(WidgetBlueprint))
+	{
+		return nullptr;
+	}
+
+	UObject* Child = FindObject<UObject>(WidgetBlueprint->WidgetTree, *Name);
+	return Cast<UWidget>(Child);
+}
+
+void UEditorTestsUtilityLibrary::SetEditorWidgetNavigationRule(UWidget* Widget, EUINavigation Nav, EUINavigationRule Rule)
+{
+	if (!Widget)
+	{
+		return;
+	}
+
+	// mimicking the FWidgetNavigationCustomization, this subobject exercises a specific edge case 
+	// within the reinstancing code, and its coverage of reinstancing that I'm interested in- not 
+	// so much editing of the subobject:
+	UWidgetNavigation* WidgetNavigation = Widget->Navigation;
+	if (!WidgetNavigation)
+	{
+		Widget->Navigation = NewObject<UWidgetNavigation>(Widget);
+		WidgetNavigation = Widget->Navigation;
+		WidgetNavigation->SetFlags(RF_Transactional);
+	}
+
+	FWidgetNavigationData* DirectionNavigation = nullptr;
+
+	switch (Nav)
+	{
+	case EUINavigation::Left:
+		DirectionNavigation = &WidgetNavigation->Left;
+		break;
+	case EUINavigation::Right:
+		DirectionNavigation = &WidgetNavigation->Right;
+		break;
+	case EUINavigation::Up:
+		DirectionNavigation = &WidgetNavigation->Up;
+		break;
+	case EUINavigation::Down:
+		DirectionNavigation = &WidgetNavigation->Down;
+		break;
+	case EUINavigation::Next:
+		DirectionNavigation = &WidgetNavigation->Next;
+		break;
+	case EUINavigation::Previous:
+		DirectionNavigation = &WidgetNavigation->Previous;
+		break;
+	default:
+		// Should not be possible.
+		check(false);
+		return;
+	}
+
+	DirectionNavigation->Rule = Rule;
+}
+
+EUINavigationRule UEditorTestsUtilityLibrary::GetEditorWidgetNavigationRule(UWidget* Widget, EUINavigation Nav)
+{
+	if (!Widget || !Widget->Navigation)
+	{
+		return EUINavigationRule::Escape;
+	}
+
+	UWidgetNavigation* WidgetNavigation = Widget->Navigation;
+
+	switch (Nav)
+	{
+	case EUINavigation::Left:
+		return WidgetNavigation->Left.Rule;
+	case EUINavigation::Right:
+		return WidgetNavigation->Right.Rule;
+	case EUINavigation::Up:
+		return WidgetNavigation->Up.Rule;
+	case EUINavigation::Down:
+		return WidgetNavigation->Down.Rule;
+	case EUINavigation::Next:
+		return WidgetNavigation->Next.Rule;
+	case EUINavigation::Previous:
+		return WidgetNavigation->Previous.Rule;
+	default:
+		// Should not be possible.
+		check(false);
+	}
+	return EUINavigationRule::Escape;
+}
+
diff --git a/Engine/Plugins/Tests/EditorTests/Source/EditorTests/Public/EditorTestsUtilityLibrary.h b/Engine/Plugins/Tests/EditorTests/Source/EditorTests/Public/EditorTestsUtilityLibrary.h
--- a/Engine/Plugins/Tests/EditorTests/Source/EditorTests/Public/EditorTestsUtilityLibrary.h
+++ b/Engine/Plugins/Tests/EditorTests/Source/EditorTests/Public/EditorTestsUtilityLibrary.h
@@ -9,6 +9,8 @@
 #include "EditorTestsUtilityLibrary.generated.h"
 
 class UMaterialOptions;
+class UWidget;
+class UWidgetBlueprint;
 struct FMeshMergingSettings;
 
 /** Blueprint library for altering and analyzing animation / skeletal data */
@@ -25,4 +27,16 @@ public:
 	/** Merges meshes and bakes out materials into a atlas-material for the given set of static mesh components using the MergeSettings */
 	UFUNCTION(BlueprintCallable, Category = "MeshMergingLibrary|Test")
 	static void MergeStaticMeshComponents(TArray<UStaticMeshComponent*> InStaticMeshComponents, const FMeshMergingSettings& MergeSettings, const bool bReplaceActors, TArray<int32>& OutLODIndices);
+
+	/** Finds a UWidget object used by the editor - useful for testing widget editing */
+	UFUNCTION(BlueprintCallable, Category = "User Interface|Test")
+	static UWidget* GetChildEditorWidgetByName(UWidgetBlueprint* WidgetBlueprint, FString Name);
+	
+	/** Simple function for setting UWidget::Navigation, which has a details customization */
+	UFUNCTION(BlueprintCallable, Category = "User Interface|Test")
+	static void SetEditorWidgetNavigationRule(UWidget* Widget, EUINavigation Nav, EUINavigationRule Rule);
+
+	/** Simple logic for getting data within UWidget::Navigation, which has a details customization */
+	UFUNCTION(BlueprintCallable, Category = "User Interface|Test")
+	static EUINavigationRule GetEditorWidgetNavigationRule(UWidget* Widget, EUINavigation Nav);
 };
\ No newline at end of file
