Fix warning V641: Memory for the buffer is allocated for pointers and not for elements of the actual type.

diff --git a/Engine/Plugins/Tests/EditorTests/Source/EditorTests/Private/UnrealEd/EditorTestActorSelection.cpp b/Engine/Plugins/Tests/EditorTests/Source/EditorTests/Private/UnrealEd/EditorTestActorSelection.cpp
--- a/Engine/Plugins/Tests/EditorTests/Source/EditorTests/Private/UnrealEd/EditorTestActorSelection.cpp
+++ b/Engine/Plugins/Tests/EditorTests/Source/EditorTests/Private/UnrealEd/EditorTestActorSelection.cpp
@@ -50,8 +50,8 @@ bool FBVerifyActorSelection::Update()
 IMPLEMENT_SIMPLE_AUTOMATION_TEST(FAssetActorTest, "Editor.Workflows.Selection.Actor", EAutomationTestFlags::EditorContext | EAutomationTestFlags::ProductFilter)
 bool FAssetActorTest::RunTest(const FString& Parameters)
 {
-	UWorld** TestWorld = (UWorld**)FMemory::Malloc(sizeof(UWorld**));
-	AActor** TestAsset = (AActor**)FMemory::Malloc(sizeof(AActor**));
+	UWorld** TestWorld = (UWorld**)FMemory::Malloc(sizeof(UWorld*)); 
+	AActor** TestAsset = (AActor**)FMemory::Malloc(sizeof(AActor*));
 
 	const FString EditorTestName = TEXT("Editor.Workflows.Selection.Actor");
 	const TCHAR* TestMapPath = TEXT("/Game/Tests/Selection/HierarchicalISMSelectionTestMap");
diff --git a/Engine/Plugins/Tests/EditorTests/Source/EditorTests/Private/UnrealEd/EditorTestHierarchicalISMSelection.cpp b/Engine/Plugins/Tests/EditorTests/Source/EditorTests/Private/UnrealEd/EditorTestHierarchicalISMSelection.cpp
--- a/Engine/Plugins/Tests/EditorTests/Source/EditorTests/Private/UnrealEd/EditorTestHierarchicalISMSelection.cpp
+++ b/Engine/Plugins/Tests/EditorTests/Source/EditorTests/Private/UnrealEd/EditorTestHierarchicalISMSelection.cpp
@@ -53,8 +53,8 @@ bool FBVerifyISMComponentSelection::Update()
 IMPLEMENT_SIMPLE_AUTOMATION_TEST(FAssetHierarchicalISMSelectionTest, "Editor.Workflows.Selection.HierarchicalISM", EAutomationTestFlags::EditorContext | EAutomationTestFlags::ProductFilter)
 bool FAssetHierarchicalISMSelectionTest::RunTest(const FString& Parameters)
 {
-	UWorld** TestWorld = (UWorld**)FMemory::Malloc(sizeof(UWorld**));
-	UInstancedStaticMeshComponent** ISMComponent = (UInstancedStaticMeshComponent**)FMemory::Malloc(sizeof(UInstancedStaticMeshComponent**));
+	UWorld** TestWorld = (UWorld**)FMemory::Malloc(sizeof(UWorld*));
+	UInstancedStaticMeshComponent** ISMComponent = (UInstancedStaticMeshComponent**)FMemory::Malloc(sizeof(UInstancedStaticMeshComponent*));
 
 	const FString EditorTestName = TEXT("Editor.Workflows.Selection.HierarchicalISM");
 	const TCHAR* TestMapPath = TEXT("/Game/Tests/Selection/HierarchicalISMSelectionTestMap");
