Adding new cvar to EditorIterationTests named project.iteration.pie.testiterations which specifies the number of times we start/run/stop PIE during the test. Default is 1 to match previous behavior.

diff --git a/Engine/Plugins/Tests/EditorTests/Source/EditorTests/Private/UnrealEd/EditorIterationTests.cpp b/Engine/Plugins/Tests/EditorTests/Source/EditorTests/Private/UnrealEd/EditorIterationTests.cpp
--- a/Engine/Plugins/Tests/EditorTests/Source/EditorTests/Private/UnrealEd/EditorIterationTests.cpp
+++ b/Engine/Plugins/Tests/EditorTests/Source/EditorTests/Private/UnrealEd/EditorIterationTests.cpp
@@ -15,6 +15,12 @@
 #include "Subsystems/AssetEditorSubsystem.h"
 #include "Editor.h"
 
+static TAutoConsoleVariable<int32> CVarProjectIterationPIETestIterations(
+	TEXT("project.iteration.pie.testiterations"),
+	1,
+	TEXT("Number of iterations of PIE to run during the Project.Iteration.PIE test.")
+);
+
 /**
  * Test to open the sub editor windows for a specified list of assets.
  * This list can be setup in the Editor Preferences window within the editor or the DefaultEngine.ini file for that particular project.
@@ -34,6 +40,8 @@ void FIterationOpenAssets::GetTests(TArray<FString>& OutBeautifiedNames, TArray<
 bool FIterationOpenAssets::RunTest(const FString& LongAssetPath)
 {
 	static uint64 FrameNumber = 0;
+	const uint64 NumFramesPerIteration = 600;
+	const uint64 NumFramesBetweenIterations = 60;
 
 	check(LongAssetPath.Len());
 
@@ -49,25 +57,42 @@ bool FIterationOpenAssets::RunTest(const FString& LongAssetPath)
 
 	AddCommand(new FFunctionLatentCommand([LongAssetPath] {
 		TRACE_BOOKMARK(TEXT("LoadAssetComplete - %s"), *LongAssetPath);
-		TRACE_BOOKMARK(TEXT("PIE - %s"), *LongAssetPath);
 		return true;
 		}));
 
-	// Do many frames of PIE (not a time span since PIE can take a while to init)
-	AddCommand(new FDelayedFunctionLatentCommand([] {
-		FrameNumber = GFrameCounter;
-		}));
-	AddCommand(new FStartPIECommand(false));
-	AddCommand(new FFunctionLatentCommand([LongAssetPath] {
-		return GFrameCounter > (FrameNumber + 600);
-		}));
-	AddCommand(new FEndPlayMapCommand());
+	int32 IterationNumber = 1;
+	int32 TotalIterations = CVarProjectIterationPIETestIterations.GetValueOnAnyThread();
+	for (; IterationNumber <= TotalIterations; IterationNumber++)
+	{
+		AddCommand(new FFunctionLatentCommand([LongAssetPath, IterationNumber] {
+			TRACE_BOOKMARK(TEXT("PIE - %s - Iteration %d"), *LongAssetPath, IterationNumber);
+			return true;
+			}));
 
-	// Teardown
-	AddCommand(new FFunctionLatentCommand([LongAssetPath] {
-		TRACE_BOOKMARK(TEXT("PIEComplete - %s"), *LongAssetPath);
-		return true;
-		}));
+		// Do many frames of PIE (not a time span since PIE can take a while to init)
+		AddCommand(new FDelayedFunctionLatentCommand([] {
+			FrameNumber = GFrameCounter;
+			}));
+		AddCommand(new FStartPIECommand(false));
+		AddCommand(new FFunctionLatentCommand([LongAssetPath] {
+			return GFrameCounter > (FrameNumber + NumFramesPerIteration);
+			}));
+		AddCommand(new FEndPlayMapCommand());
+
+		// Teardown
+		AddCommand(new FFunctionLatentCommand([LongAssetPath, IterationNumber] {
+			TRACE_BOOKMARK(TEXT("PIEComplete - %s - Iteration %d"), *LongAssetPath, IterationNumber);
+			return true;
+			}));
+
+		if (IterationNumber < TotalIterations)
+		{
+			// Pause between iterations
+			AddCommand(new FFunctionLatentCommand([LongAssetPath] {
+				return GFrameCounter > (FrameNumber + NumFramesPerIteration + NumFramesBetweenIterations);
+				}));
+		}
+	}
 
 	// Wait on all async asset processing
 	AddCommand(new FFunctionLatentCommand([] {
