Compilation of Blueprint with Missing Node Class Test

- Added "Project.Functional Tests.Tests.Blueprints.Compile Missing Node" in EditorTests Plugin for compilation of a Blueprint that had a missing node from class removal

#rnx

#virtualized

diff --git a/Engine/Plugins/Tests/EditorTests/Source/EditorTests/Private/UnrealEd/BlueprintAutomationTests.cpp b/Engine/Plugins/Tests/EditorTests/Source/EditorTests/Private/UnrealEd/BlueprintAutomationTests.cpp
--- a/Engine/Plugins/Tests/EditorTests/Source/EditorTests/Private/UnrealEd/BlueprintAutomationTests.cpp
+++ b/Engine/Plugins/Tests/EditorTests/Source/EditorTests/Private/UnrealEd/BlueprintAutomationTests.cpp
@@ -51,6 +51,8 @@ DEFINE_LOG_CATEGORY_STATIC(LogBlueprintAutomationTests, Log, All);
 IMPLEMENT_COMPLEX_AUTOMATION_TEST(FCompileBlueprintsTest, "Project.Blueprints.Compile Blueprints", EAutomationTestFlags::EditorContext | EAutomationTestFlags::StressFilter)
 IMPLEMENT_COMPLEX_AUTOMATION_TEST(FCompileAnimBlueprintsTest, "Project.Blueprints.Compile Anims", EAutomationTestFlags::EditorContext | EAutomationTestFlags::StressFilter)
 
+IMPLEMENT_SIMPLE_AUTOMATION_TEST(FCompileMissingNodeBlueprint, "Project.Functional Tests.Tests.Blueprints.Compile Missing Node", EAutomationTestFlags::EditorContext | EAutomationTestFlags::EngineFilter)
+
 class FBlueprintAutomationTestUtilities
 {
 	/** An incrementing number that can be used to tack on to save files, etc. (for avoiding naming conflicts)*/
@@ -858,3 +860,12 @@ bool FCompileAnimBlueprintsTest::RunTest(const FString& Parameters)
 	return FBlueprintAutomationTestUtilities::CompileBlueprint(Parameters);
 }
 
+/************************************************************************/
+/* FCompileMissingNodeBlueprint                                         */
+/************************************************************************/
+
+/** Tests compilation of a Blueprint object that had a missing node from class removal - C12694393 */
+bool FCompileMissingNodeBlueprint::RunTest(const FString& Parameters)
+{
+	return FBlueprintAutomationTestUtilities::CompileBlueprint(TEXT("/EditorTests/Blueprints/BP_NodeClassIsMissing.BP_NodeClassIsMissing"));
+}
\ No newline at end of file
