EngineTest - Added new actor merging test case which tries to merge meshes with different configurations that include empty sections

#virtualized

diff --git a/Engine/Plugins/Tests/EditorTests/Source/EditorTests/Private/UnrealEd/EditorTestsUtilityLibrary.cpp b/Engine/Plugins/Tests/EditorTests/Source/EditorTests/Private/UnrealEd/EditorTestsUtilityLibrary.cpp
--- a/Engine/Plugins/Tests/EditorTests/Source/EditorTests/Private/UnrealEd/EditorTestsUtilityLibrary.cpp
+++ b/Engine/Plugins/Tests/EditorTests/Source/EditorTests/Private/UnrealEd/EditorTestsUtilityLibrary.cpp
@@ -23,6 +23,33 @@
 #include "Blueprint/WidgetTree.h"
 #include "Components/Widget.h"
 
+DECLARE_LOG_CATEGORY_CLASS(LogEditorTestsUtilityLibrary, Log, All)
+
+namespace
+{
+	void WaitForTextures(const UMaterialInterface* Material)
+	{
+		TArray<UTexture*> MaterialTextures;
+		Material->GetUsedTextures(MaterialTextures, EMaterialQualityLevel::Num, true, GMaxRHIFeatureLevel, true);
+
+		FTextureCompilingManager::Get().FinishCompilation(MaterialTextures);
+
+		// Force load materials used by the current material
+		for (UTexture* Texture : MaterialTextures)
+		{
+			if (Texture != NULL)
+			{
+				UTexture2D* Texture2D = Cast<UTexture2D>(Texture);
+				if (Texture2D)
+				{
+					Texture2D->SetForceMipLevelsToBeResident(30.0f);
+					Texture2D->WaitForStreaming();
+				}
+			}
+		}
+	}
+}
+
 void UEditorTestsUtilityLibrary::BakeMaterialsForComponent(UStaticMeshComponent* InStaticMeshComponent, const UMaterialOptions* MaterialOptions, const UMaterialMergeOptions* MaterialMergeOptions)
 {
 	if (InStaticMeshComponent != nullptr && InStaticMeshComponent->GetStaticMesh() != nullptr)
@@ -48,68 +75,35 @@ void UEditorTestsUtilityLibrary::BakeMaterialsForComponent(UStaticMeshComponent*
 		for (int32 MaterialIndex = 0; MaterialIndex < NumMaterials; ++MaterialIndex)
 		{
 			UMaterialInterface* Material = InStaticMeshComponent->GetMaterial(MaterialIndex);
-			TArray<UTexture*> MaterialTextures;
-			Material->GetUsedTextures(MaterialTextures, EMaterialQualityLevel::Num, true, GMaxRHIFeatureLevel, true);
-
-			FTextureCompilingManager::Get().FinishCompilation(MaterialTextures);
-
-			// Force load materials used by the current material
-			for (UTexture* Texture : MaterialTextures)
-			{
-				if (Texture != NULL)
-				{
-					UTexture2D* Texture2D = Cast<UTexture2D>(Texture);
-					if (Texture2D)
-					{
-						Texture2D->SetForceMipLevelsToBeResident(30.0f);
-						Texture2D->WaitForStreaming();
-					}
-				}
-			}
+			WaitForTextures(Material);
 		}
 	}
 }
 
-void UEditorTestsUtilityLibrary::MergeStaticMeshComponents(TArray<UStaticMeshComponent*> InStaticMeshComponents, const FMeshMergingSettings& MergeSettings, const bool bReplaceActors, TArray<int32>& OutLODIndices)
+void UEditorTestsUtilityLibrary::MergeStaticMeshComponents(const TArray<UStaticMeshComponent*>& InStaticMeshComponents, const FMeshMergingSettings& MergeSettings, const bool bReplaceActors, TArray<int32>& OutLODIndices)
 {
 	FModuleManager::Get().LoadModule("MaterialBaking");
 
-	InStaticMeshComponents.RemoveAll([](UStaticMeshComponent* Component) { return Component == nullptr || Component->GetStaticMesh() == nullptr; });
-	if (InStaticMeshComponents.Num() && InStaticMeshComponents[0] != nullptr)
-	{
-		const IMeshMergeUtilities& MeshMergeUtilities = FModuleManager::Get().LoadModuleChecked<IMeshMergeModule>("MeshMergeUtilities").GetUtilities();
-		UWorld* World = InStaticMeshComponents[0]->GetWorld();
+	const IMeshMergeUtilities& MeshMergeUtilities = FModuleManager::Get().LoadModuleChecked<IMeshMergeModule>("MeshMergeUtilities").GetUtilities();
 
-		// Convert array of StaticMeshComponents to PrimitiveComponents
-		TArray<UPrimitiveComponent*> PrimCompsToMerge;
-		Algo::Transform(InStaticMeshComponents, PrimCompsToMerge, [](UStaticMeshComponent* StaticMeshComp) { return StaticMeshComp; });
+	// Convert array of StaticMeshComponents to PrimitiveComponents
+	TArray<UPrimitiveComponent*> PrimCompsToMerge;
+	Algo::TransformIf(InStaticMeshComponents, PrimCompsToMerge, 
+		[](UStaticMeshComponent* StaticMeshComp) { return StaticMeshComp == nullptr || StaticMeshComp->GetStaticMesh() == nullptr; },
+		[](UStaticMeshComponent* StaticMeshComp) { return StaticMeshComp; });
+
+	if (!PrimCompsToMerge.IsEmpty())
+	{
+		UWorld* World = PrimCompsToMerge[0]->GetWorld();
 
 		TArray<UObject*> Output;
 		FVector OutPosition;
 		MeshMergeUtilities.MergeComponentsToStaticMesh(PrimCompsToMerge, World, MergeSettings, nullptr, GetTransientPackage(), InStaticMeshComponents[0]->GetStaticMesh()->GetName(), Output, OutPosition, 1.0f, false);
-		
+
 		UObject** MaterialPtr = Output.FindByPredicate([](UObject* Object) { return Object->IsA<UMaterial>(); });
-		if (MaterialPtr)
+		if (UMaterial* MergedMaterial = Cast<UMaterial>(*MaterialPtr))
 		{
-			UMaterial* MergedMaterial = Cast<UMaterial>(*MaterialPtr);
-			TArray<UTexture*> MaterialTextures;
-			MergedMaterial->GetUsedTextures(MaterialTextures, EMaterialQualityLevel::Num, true, GMaxRHIFeatureLevel, true);
-
-			FTextureCompilingManager::Get().FinishCompilation(MaterialTextures);
-
-			// Force load materials used by the current material
-			for (UTexture* Texture : MaterialTextures)
-			{
-				if (Texture != NULL)
-				{
-					UTexture2D* Texture2D = Cast<UTexture2D>(Texture);
-					if (Texture2D)
-					{
-						Texture2D->SetForceMipLevelsToBeResident(30.0f);
-						Texture2D->WaitForStreaming();
-					}
-				}
-			}
+			WaitForTextures(MergedMaterial);
 		}
 
 		// Place new mesh in the world
@@ -118,9 +112,7 @@ void UEditorTestsUtilityLibrary::MergeStaticMeshComponents(TArray<UStaticMeshCom
 			UObject** ObjectPtr = Output.FindByPredicate([](const UObject* Object) {return Object->IsA<UStaticMesh>(); });
 			if (ObjectPtr != nullptr && *ObjectPtr != nullptr)
 			{
-				UStaticMesh* MergedMesh = CastChecked<UStaticMesh>(*ObjectPtr);
-
-				if (MergedMesh)
+				if (UStaticMesh* MergedMesh = CastChecked<UStaticMesh>(*ObjectPtr))
 				{
 					for (int32 Index = 0; Index < MergedMesh->GetNumLODs(); ++Index)
 					{
@@ -153,6 +145,66 @@ void UEditorTestsUtilityLibrary::MergeStaticMeshComponents(TArray<UStaticMeshCom
 	}
 }
 
+void UEditorTestsUtilityLibrary::CreateProxyMesh(const TArray<UStaticMeshComponent*>& InStaticMeshComponents, const struct FMeshProxySettings& ProxySettings)
+{
+	TArray<UStaticMeshComponent*> StaticMeshComponents;
+	Algo::CopyIf(InStaticMeshComponents, StaticMeshComponents, [](UStaticMeshComponent* StaticMeshComp) { return StaticMeshComp != nullptr && StaticMeshComp->GetStaticMesh() != nullptr; });
+
+	if (StaticMeshComponents.IsEmpty())
+	{
+		return;
+	}
+
+	UWorld* World = StaticMeshComponents[0]->GetWorld();
+
+	// Generate proxy mesh and proxy material assets 
+	FCreateProxyDelegate ProxyDelegate;
+	ProxyDelegate.BindLambda([&](const FGuid Guid, TArray<UObject*>& AssetsToSync)
+	{
+		UStaticMesh* ProxyMesh = nullptr;
+		if (!AssetsToSync.FindItemByClass(&ProxyMesh))
+		{
+			UE_LOG(LogEditorTestsUtilityLibrary, Error, TEXT("CreateProxyMesh failed. No mesh was created."));
+			return;
+		}
+
+		UMaterial* ProxyMaterial = nullptr;
+		if (!AssetsToSync.FindItemByClass(&ProxyMaterial))
+		{
+			WaitForTextures(ProxyMaterial);
+		}
+
+		// Place new mesh in the world (on a new actor)
+		FActorSpawnParameters Params;
+		Params.OverrideLevel = World->PersistentLevel;
+		AStaticMeshActor* MergedActor = World->SpawnActor<AStaticMeshActor>(Params);
+		if (!MergedActor)
+		{
+			UE_LOG(LogEditorTestsUtilityLibrary, Error, TEXT("CreateProxyMesh failed. Internal error while creating the merged actor."));
+			return;
+		}
+
+		MergedActor->SetMobility(EComponentMobility::Movable);
+		MergedActor->SetActorLabel(TEXT("Proxy_Actor"));
+		MergedActor->GetStaticMeshComponent()->SetStaticMesh(ProxyMesh);		
+
+		TArray<AActor*> OwningActors;
+		for (UStaticMeshComponent* Component : StaticMeshComponents)
+		{
+			OwningActors.AddUnique(Component->GetOwner());
+		}
+
+		// Remove source actors
+		for (AActor* Actor : OwningActors)
+		{
+			Actor->Destroy();
+		}
+	});
+
+	const IMeshMergeUtilities& MeshMergeUtilities = FModuleManager::Get().LoadModuleChecked<IMeshMergeModule>("MeshMergeUtilities").GetUtilities();
+	MeshMergeUtilities.CreateProxyMesh(StaticMeshComponents, ProxySettings, GetTransientPackage(), TEXT("ProxyMesh"), FGuid::NewGuid(), ProxyDelegate);
+}
+
 UWidget* UEditorTestsUtilityLibrary::GetChildEditorWidgetByName(UWidgetBlueprint* WidgetBlueprint, FString Name)
 {
 	if (!ensure(WidgetBlueprint))
diff --git a/Engine/Plugins/Tests/EditorTests/Source/EditorTests/Public/EditorTestsUtilityLibrary.h b/Engine/Plugins/Tests/EditorTests/Source/EditorTests/Public/EditorTestsUtilityLibrary.h
--- a/Engine/Plugins/Tests/EditorTests/Source/EditorTests/Public/EditorTestsUtilityLibrary.h
+++ b/Engine/Plugins/Tests/EditorTests/Source/EditorTests/Public/EditorTestsUtilityLibrary.h
@@ -4,7 +4,9 @@
 
 #include "Kismet/BlueprintFunctionLibrary.h"
 #include "Components/StaticMeshComponent.h"
+
 #include "MeshMerge/MeshMergingSettings.h"
+#include "MeshMerge/MeshProxySettings.h"
 
 #include "EditorTestsUtilityLibrary.generated.h"
 
@@ -26,12 +28,16 @@ public:
 
 	/** Merges meshes and bakes out materials into a atlas-material for the given set of static mesh components using the MergeSettings */
 	UFUNCTION(BlueprintCallable, Category = "MeshMergingLibrary|Test")
-	static void MergeStaticMeshComponents(TArray<UStaticMeshComponent*> InStaticMeshComponents, const FMeshMergingSettings& MergeSettings, const bool bReplaceActors, TArray<int32>& OutLODIndices);
+	static void MergeStaticMeshComponents(const TArray<UStaticMeshComponent*>& InStaticMeshComponents, const FMeshMergingSettings& MergeSettings, const bool bReplaceActors, TArray<int32>& OutLODIndices);
+
+	/** Simplify meshes and bakes out materials into a atlas-material for the given set of static mesh components using the ProxySettings */
+	UFUNCTION(BlueprintCallable, Category = "MeshMergingLibrary|Test")
+	static void CreateProxyMesh(const TArray<UStaticMeshComponent*>& InStaticMeshComponents, const struct FMeshProxySettings& ProxySettings);
 
 	/** Finds a UWidget object used by the editor - useful for testing widget editing */
 	UFUNCTION(BlueprintCallable, Category = "User Interface|Test")
 	static UWidget* GetChildEditorWidgetByName(UWidgetBlueprint* WidgetBlueprint, FString Name);
-	
+
 	/** Simple function for setting UWidget::Navigation, which has a details customization */
 	UFUNCTION(BlueprintCallable, Category = "User Interface|Test")
 	static void SetEditorWidgetNavigationRule(UWidget* Widget, EUINavigation Nav, EUINavigationRule Rule);
@@ -44,4 +50,4 @@ public:
 
 #if UE_ENABLE_INCLUDE_ORDER_DEPRECATED_IN_5_5
 #include "Engine/MeshMerging.h"
-#endif
\ No newline at end of file
+#endif
