Update UEditorTestsUtilityLibrary::CreateProxyMesh to support Material Instances

diff --git a/Engine/Plugins/Tests/EditorTests/Source/EditorTests/Private/UnrealEd/EditorTestsUtilityLibrary.cpp b/Engine/Plugins/Tests/EditorTests/Source/EditorTests/Private/UnrealEd/EditorTestsUtilityLibrary.cpp
--- a/Engine/Plugins/Tests/EditorTests/Source/EditorTests/Private/UnrealEd/EditorTestsUtilityLibrary.cpp
+++ b/Engine/Plugins/Tests/EditorTests/Source/EditorTests/Private/UnrealEd/EditorTestsUtilityLibrary.cpp
@@ -168,7 +168,7 @@ void UEditorTestsUtilityLibrary::CreateProxyMesh(const TArray<UStaticMeshCompone
 			return;
 		}
 
-		UMaterial* ProxyMaterial = nullptr;
+		UMaterialInterface* ProxyMaterial = nullptr;
 		if (!AssetsToSync.FindItemByClass(&ProxyMaterial))
 		{
 			WaitForTextures(ProxyMaterial);
