Fix an issue in UEditorTestsUtilityLibrary::MergeStaticMeshComponents preventing Static Mesh Components from being merged

diff --git a/Engine/Plugins/Tests/EditorTests/Source/EditorTests/Private/UnrealEd/EditorTestsUtilityLibrary.cpp b/Engine/Plugins/Tests/EditorTests/Source/EditorTests/Private/UnrealEd/EditorTestsUtilityLibrary.cpp
--- a/Engine/Plugins/Tests/EditorTests/Source/EditorTests/Private/UnrealEd/EditorTestsUtilityLibrary.cpp
+++ b/Engine/Plugins/Tests/EditorTests/Source/EditorTests/Private/UnrealEd/EditorTestsUtilityLibrary.cpp
@@ -89,7 +89,7 @@ void UEditorTestsUtilityLibrary::MergeStaticMeshComponents(const TArray<UStaticM
 	// Convert array of StaticMeshComponents to PrimitiveComponents
 	TArray<UPrimitiveComponent*> PrimCompsToMerge;
 	Algo::TransformIf(InStaticMeshComponents, PrimCompsToMerge, 
-		[](UStaticMeshComponent* StaticMeshComp) { return StaticMeshComp == nullptr || StaticMeshComp->GetStaticMesh() == nullptr; },
+		[](UStaticMeshComponent* StaticMeshComp) { return StaticMeshComp != nullptr && StaticMeshComp->GetStaticMesh() != nullptr; },
 		[](UStaticMeshComponent* StaticMeshComp) { return StaticMeshComp; });
 
 	if (!PrimCompsToMerge.IsEmpty())
@@ -101,9 +101,12 @@ void UEditorTestsUtilityLibrary::MergeStaticMeshComponents(const TArray<UStaticM
 		MeshMergeUtilities.MergeComponentsToStaticMesh(PrimCompsToMerge, World, MergeSettings, nullptr, GetTransientPackage(), InStaticMeshComponents[0]->GetStaticMesh()->GetName(), Output, OutPosition, 1.0f, false);
 
 		UObject** MaterialPtr = Output.FindByPredicate([](UObject* Object) { return Object->IsA<UMaterial>(); });
-		if (UMaterial* MergedMaterial = Cast<UMaterial>(*MaterialPtr))
+		if (MaterialPtr)
 		{
-			WaitForTextures(MergedMaterial);
+			if (UMaterial* MergedMaterial = Cast<UMaterial>(*MaterialPtr))
+			{
+				WaitForTextures(MergedMaterial);
+			}
 		}
 
 		// Place new mesh in the world
