JsonUtilities JsonSchema tests.

In EditorTests plugin.

#Synced-CL 49091575

diff --git a/Engine/Plugins/Tests/EditorTests/Source/EditorTests/EditorTests.Build.cs b/Engine/Plugins/Tests/EditorTests/Source/EditorTests/EditorTests.Build.cs
--- a/Engine/Plugins/Tests/EditorTests/Source/EditorTests/EditorTests.Build.cs
+++ b/Engine/Plugins/Tests/EditorTests/Source/EditorTests/EditorTests.Build.cs
@@ -30,10 +30,11 @@ public class EditorTests : ModuleRules
 				"AssetTools",
 				"MainFrame",
 				"MaterialEditor",
+				"JsonUtilitiesEditor",
 				"JsonUtilities",
+				"Json",
 				"Analytics",
-				"ContentBrowser",
-				
+				"ContentBrowser",				
 				"SourceControl",
 				"RHI",
 				"BlueprintGraph",
diff --git a/Engine/Plugins/Tests/EditorTests/Source/EditorTests/Private/UnrealEd/JsonUtilitiesTests.cpp b/Engine/Plugins/Tests/EditorTests/Source/EditorTests/Private/UnrealEd/JsonUtilitiesTests.cpp
new file mode 100644
--- /dev/null
+++ b/Engine/Plugins/Tests/EditorTests/Source/EditorTests/Private/UnrealEd/JsonUtilitiesTests.cpp
@@ -0,0 +1,729 @@
+// Copyright Epic Games, Inc. All Rights Reserved. 
+
+
+#include "JsonUtilitiesTests.h"
+
+#include "Modules/ModuleManager.h"
+
+#if WITH_EDITOR
+#include "Tests/AutomationEditorCommon.h"
+#endif
+
+#include "JsonSchema/JsonSchemaGenerator.h"
+
+#if WITH_DEV_AUTOMATION_TESTS
+
+// JsonSchema functions are currently marker 'deprecated' as they are still 'experimental'.
+// Remove this later once they are no longer deprecated..
+PRAGMA_DISABLE_DEPRECATION_WARNINGS
+
+#define PRINT_JSONS_TO_LOG 1
+
+namespace
+{
+	// Gets last member of the member path.
+	FString Last(FJsonSchemaMemberPath& MemberPath)
+	{
+		const FString MemberPathString = MemberPath.ToString();
+	
+		if (int32 Index = INDEX_NONE;
+			MemberPathString.FindLastChar(FJsonSchemaMemberPath::Delimiter[0], Index))
+		{
+			return MemberPathString.Mid(Index + 1);
+		}		
+
+		return MemberPathString;
+	}
+	
+	bool TestJsonSchemaIsValidAndNotEmpty(TSharedPtr<FJsonObject> Schema, FAutomationTestBase& Test)
+	{
+		if (!Test.TestTrue("Schema is valid.", Schema.IsValid()) ||
+			!Test.TestTrue("Schema is not empty (has values.)", Schema->Values.Num() > 0))
+		{
+			return false;
+		}
+		
+		return true;
+	}
+	
+	TOptional<TSharedPtr<FJsonObject>> TestAndGetJsonObjectField(TSharedPtr<FJsonObject> Schema, const FString& FieldName, FAutomationTestBase& Test)
+	{
+		TOptional<TSharedPtr<FJsonObject>> OptionalJsonObject;
+		
+		if (Test.TestTrue("JSON schema is valid.", Schema.IsValid()))
+		{
+			if (!Test.TestTrue(
+				FString::Printf(TEXT("%s object exists."), *FieldName),
+				Schema->HasTypedField(*FieldName, EJson::Object)))
+			{
+				return OptionalJsonObject;
+			}
+
+			OptionalJsonObject = Schema->GetObjectField(*FieldName);
+		}
+		
+		return OptionalJsonObject;
+	}
+	
+	TOptional<FString> TestAndGetJsonStringField(TSharedPtr<FJsonObject> Schema, const FString& FieldName, FAutomationTestBase& Test)
+	{
+		TOptional<FString> OptionalJsonString;
+		if (FString JsonString;
+			Test.TestTrue(
+				FString::Printf(TEXT("%s string exists."), *FieldName),
+				Schema->TryGetStringField(*FieldName, JsonString)))
+		{
+			OptionalJsonString = JsonString;
+		}
+		
+		return OptionalJsonString;
+	}
+	
+	TOptional<TArray<FString>> TestAndGetJsonStringArrayField(TSharedPtr<FJsonObject> Schema, const FString& FieldName, FAutomationTestBase& Test)
+	{
+		TOptional<TArray<FString>> OptionalJsonStringArray;
+		if (TArray<FString> JsonStringArray;
+			Test.TestTrue(
+				FString::Printf(TEXT("%s string array exists."), *FieldName),
+				Schema->TryGetStringArrayField(*FieldName, JsonStringArray)))
+		{
+			OptionalJsonStringArray = JsonStringArray;
+		}
+		
+		return OptionalJsonStringArray;
+	}
+	
+	TOptional<float> TestAndGetJsonFloatNumberField(TSharedPtr<FJsonObject> Schema, const FString& FieldName, FAutomationTestBase& Test)
+	{
+		TOptional<float> OptionalJsonNumber;
+		if (float JsonNumber = 0.0f;
+			Test.TestTrue(
+				FString::Printf(TEXT("%s float number exists."), *FieldName),
+				Schema->TryGetNumberField(*FieldName, JsonNumber)))
+		{
+			OptionalJsonNumber = JsonNumber;
+		}
+		
+		return OptionalJsonNumber;
+	}
+	
+	TOptional<int32> TestAndGetJsonIntNumberField(TSharedPtr<FJsonObject> Schema, const FString& FieldName, FAutomationTestBase& Test)
+	{
+		TOptional<int32> OptionalJsonNumber;
+		if (int32 JsonNumber = 0;
+			Test.TestTrue(
+				FString::Printf(TEXT("%s int number exists."), *FieldName),
+				Schema->TryGetNumberField(*FieldName, JsonNumber)))
+		{
+			OptionalJsonNumber = JsonNumber;
+		}
+		
+		return OptionalJsonNumber;
+	}
+	
+	void MaybePrintJsonToLog(const TCHAR* Name, const TSharedRef<FJsonObject>& JsonObject)
+	{
+#if PRINT_JSONS_TO_LOG
+		FString JsonString;
+		const TSharedRef<TJsonWriter<>> JsonWriter = TJsonWriterFactory<>::Create(&JsonString);
+		FJsonSerializer::Serialize(JsonObject, JsonWriter);
+		UE_LOG(LogJson, Warning, TEXT("[%s]\n%s"), Name, *JsonString);
+#endif
+	}
+}
+
+BEGIN_DEFINE_SPEC(FJsonUtiltiesEditorSpec, "System.Promotion.Editor.JsonUtilities",
+	EAutomationTestFlags::EditorContext|EAutomationTestFlags::EngineFilter)
+END_DEFINE_SPEC(FJsonUtiltiesEditorSpec)
+
+void FJsonUtiltiesEditorSpec::Define()
+{
+	BeforeEach([this]() -> void
+		{
+			// Loading this module will enable the collection of metadata for the schemas.
+			static const TCHAR* JsonUtilitiesEditorModuleName = TEXT("JsonUtilitiesEditor");
+			const IModuleInterface* Module = FModuleManager::Get().LoadModule(JsonUtilitiesEditorModuleName);
+			TestTrue(
+				FString::Printf(TEXT("Load module '%s'."), JsonUtilitiesEditorModuleName), 
+				LIKELY(Module));
+		});
+	
+	Describe("Should generate JSON schemas from UClasses/UStructs/UProperties/UFunctions.", [this]() -> void
+		{
+			It("Should generate JSON schema from a UStruct - smoke test.", [this]() -> void
+				{
+					const TSharedPtr<FJsonObject> Schema = FJsonSchemaGenerator::UStructToJsonSchemaObject<FJsonUtilitiesFakeStruct>();
+					if (!TestJsonSchemaIsValidAndNotEmpty(Schema, *this))
+					{
+						return;
+					}
+
+					MaybePrintJsonToLog(*FJsonUtilitiesFakeStruct::StaticStruct()->GetName(), Schema.ToSharedRef());
+				});
+			
+			It("Should generate JSON schema from a UClass - smoke test.", [this]() -> void
+				{
+					const TSharedPtr<FJsonObject> Schema = FJsonSchemaGenerator::UStructToJsonSchemaObject<UJsonUtilitiesFakeClass>();
+					if (!TestJsonSchemaIsValidAndNotEmpty(Schema, *this))
+					{
+						return;
+					}
+
+					MaybePrintJsonToLog(*UJsonUtilitiesFakeClass::StaticClass()->GetName(), Schema.ToSharedRef());
+				});
+			
+			It("Should generate JSON schema from a UStruct's float UProperty.", [this]() -> void
+				{
+					FJsonSchemaMemberPath MemberPath(TEXT("[schema]")); // ..mosty for error reporting
+					
+					const TSharedPtr<FJsonObject> Schema = FJsonSchemaGenerator::UStructToJsonSchemaObject<FJsonUtilitiesFakeStruct>();
+					if (!TestJsonSchemaIsValidAndNotEmpty(Schema, *this))
+					{
+						return;
+					}
+						
+					UE_JSON_SCHEMA_SCOPED_MEMBER_PATH_PUSH(MemberPath, TEXT("properties"), true);
+					
+					if (const TOptional<TSharedPtr<FJsonObject>> Schema_Properties = TestAndGetJsonObjectField(
+							Schema, Last(MemberPath), *this);
+						Schema_Properties.IsSet())
+					{
+						UE_JSON_SCHEMA_SCOPED_MEMBER_PATH_PUSH(MemberPath, TEXT("testFloat"), true);
+						
+						if (const TOptional<TSharedPtr<FJsonObject>> Schema_Properties_TestFloat = TestAndGetJsonObjectField(
+								*Schema_Properties, Last(MemberPath), *this);
+							Schema_Properties_TestFloat.IsSet())
+						{
+							{
+								UE_JSON_SCHEMA_SCOPED_MEMBER_PATH_PUSH(MemberPath, TEXT("type"), true);
+								
+								if (const TOptional<FString> Schema_Properties_TestFloat_Type = TestAndGetJsonStringField(
+										*Schema_Properties_TestFloat, Last(MemberPath), *this);
+									Schema_Properties_TestFloat_Type.IsSet())
+								{
+									TestEqual(
+										FString::Printf(TEXT("%s is number."), *MemberPath), 
+											(*Schema_Properties_TestFloat_Type), 
+										"number");
+								}
+							}
+
+							{
+								UE_JSON_SCHEMA_SCOPED_MEMBER_PATH_PUSH(MemberPath, TEXT("minimum"), true);
+								
+								if (const TOptional<float> Schema_Properties_TestFloat_Minimum = TestAndGetJsonFloatNumberField(
+										*Schema_Properties_TestFloat, Last(MemberPath), *this);
+									Schema_Properties_TestFloat_Minimum.IsSet())
+								{
+									TestEqual(
+										FString::Printf(TEXT("%s should be %f. Is %f."), *MemberPath, 
+											ExpectedTestFloatMinimum, *Schema_Properties_TestFloat_Minimum),
+										*Schema_Properties_TestFloat_Minimum, 
+										ExpectedTestFloatMinimum); 
+								}
+							}
+							
+							{
+								UE_JSON_SCHEMA_SCOPED_MEMBER_PATH_PUSH(MemberPath, TEXT("maximum"), true);
+								
+								if (const TOptional<float> Schema_Properties_TestFloat_Maximum = TestAndGetJsonFloatNumberField(
+										*Schema_Properties_TestFloat, Last(MemberPath), *this);
+									Schema_Properties_TestFloat_Maximum.IsSet())
+								{
+									TestEqual(
+										FString::Printf(TEXT("%s should be %f. Is %f."), *MemberPath, 
+											ExpectedTestFloatMaximum,* Schema_Properties_TestFloat_Maximum),
+										*Schema_Properties_TestFloat_Maximum, 
+										ExpectedTestFloatMaximum); 
+								}
+							}
+							
+							{
+								UE_JSON_SCHEMA_SCOPED_MEMBER_PATH_PUSH(MemberPath, TEXT("description"), true);
+								
+								if (const TOptional<FString> Schema_Properties_TestFloat_Description = TestAndGetJsonStringField(
+										*Schema_Properties_TestFloat, Last(MemberPath), *this);
+									Schema_Properties_TestFloat_Description.IsSet())
+								{
+									TestEqual(
+										FString::Printf(TEXT("%s should be '%s'. Is '%s'."), *MemberPath, 
+											*ExpectedTestFloatDescription, **Schema_Properties_TestFloat_Description),
+										*Schema_Properties_TestFloat_Description, 
+										ExpectedTestFloatDescription); 
+								}
+							}
+						}
+					}
+				});
+			
+			It("Should generate JSON schema from a UStruct's int UProperty.", [this]() -> void
+				{
+					FJsonSchemaMemberPath MemberPath(TEXT("[schema]")); // ..mosty for error reporting
+					
+					const TSharedPtr<FJsonObject> Schema = FJsonSchemaGenerator::UStructToJsonSchemaObject<FJsonUtilitiesFakeStruct>();
+					if (!TestJsonSchemaIsValidAndNotEmpty(Schema, *this))
+					{
+						return;
+					}
+						
+					UE_JSON_SCHEMA_SCOPED_MEMBER_PATH_PUSH(MemberPath, TEXT("properties"), true);
+					
+					if (const TOptional<TSharedPtr<FJsonObject>> Schema_Properties = TestAndGetJsonObjectField(
+							Schema, Last(MemberPath), *this);
+						Schema_Properties.IsSet())
+					{
+						UE_JSON_SCHEMA_SCOPED_MEMBER_PATH_PUSH(MemberPath, TEXT("testInt"), true);
+						
+						if (const TOptional<TSharedPtr<FJsonObject>> Schema_Properties_TestInt = TestAndGetJsonObjectField(
+								*Schema_Properties, Last(MemberPath), *this);
+							Schema_Properties_TestInt.IsSet())
+						{
+							{
+								UE_JSON_SCHEMA_SCOPED_MEMBER_PATH_PUSH(MemberPath, TEXT("type"), true);
+								
+								if (const TOptional<FString> Schema_Properties_TestInt_Type = TestAndGetJsonStringField(
+										*Schema_Properties_TestInt, Last(MemberPath), *this);
+									Schema_Properties_TestInt_Type.IsSet())
+								{
+									TestTrue(
+										FString::Printf(TEXT("%s is integer."), *MemberPath),
+										*Schema_Properties_TestInt_Type == "integer");
+								}
+							}
+
+							{
+								UE_JSON_SCHEMA_SCOPED_MEMBER_PATH_PUSH(MemberPath, TEXT("minimum"), true);
+								
+								if (const TOptional<int32> Schema_Properties_TestInt_Minimum = TestAndGetJsonIntNumberField(
+										*Schema_Properties_TestInt, Last(MemberPath), *this);
+									Schema_Properties_TestInt_Minimum.IsSet())
+								{
+									TestEqual(
+										FString::Printf(TEXT("%s should be %d. Is %d."), *MemberPath, 
+											ExpectedTestIntMinimum, *Schema_Properties_TestInt_Minimum),
+										*Schema_Properties_TestInt_Minimum, 
+										ExpectedTestIntMinimum); 
+								}
+							}
+							
+							{
+								UE_JSON_SCHEMA_SCOPED_MEMBER_PATH_PUSH(MemberPath, TEXT("maximum"), true);
+								
+								if (const TOptional<int32> Schema_Properties_TestInt_Maximum = TestAndGetJsonIntNumberField(
+										*Schema_Properties_TestInt, Last(MemberPath), *this);
+									Schema_Properties_TestInt_Maximum.IsSet())
+								{
+									TestEqual(
+										FString::Printf(TEXT("%s should be %d. Is %d."), *MemberPath, 
+											ExpectedTestIntMaximum, *Schema_Properties_TestInt_Maximum),
+										*Schema_Properties_TestInt_Maximum, 
+										ExpectedTestIntMaximum); 
+								}
+							}
+							
+							{
+								UE_JSON_SCHEMA_SCOPED_MEMBER_PATH_PUSH(MemberPath, TEXT("description"), true);
+								
+								if (const TOptional<FString> Schema_Properties_TestInt_Description = TestAndGetJsonStringField(
+										*Schema_Properties_TestInt, Last(MemberPath), *this);
+									Schema_Properties_TestInt_Description.IsSet())
+								{
+									TestEqual(
+										FString::Printf(TEXT("%s should be '%s'. Is '%s'."), *MemberPath, 
+											*ExpectedTestIntDescription, **Schema_Properties_TestInt_Description),
+										*Schema_Properties_TestInt_Description, 
+										ExpectedTestIntDescription); 
+								}
+							}
+						}
+					}
+				});
+			
+			It("Should generate JSON schema from a UStruct's bool UProperty.", [this]() -> void
+				{
+					FJsonSchemaMemberPath MemberPath(TEXT("[schema]")); // ..mosty for error reporting
+					
+					const TSharedPtr<FJsonObject> Schema = FJsonSchemaGenerator::UStructToJsonSchemaObject<FJsonUtilitiesFakeStruct>();
+					if (!TestJsonSchemaIsValidAndNotEmpty(Schema, *this))
+					{
+						return;
+					}
+
+					UE_JSON_SCHEMA_SCOPED_MEMBER_PATH_PUSH(MemberPath, TEXT("properties"), true);
+					
+					if (const TOptional<TSharedPtr<FJsonObject>> Schema_Properties = TestAndGetJsonObjectField(
+							Schema, Last(MemberPath), *this);
+						Schema_Properties.IsSet())
+					{
+						UE_JSON_SCHEMA_SCOPED_MEMBER_PATH_PUSH(MemberPath, TEXT("testBool"), true);
+						
+						if (const TOptional<TSharedPtr<FJsonObject>> Schema_Properties_TestBool = TestAndGetJsonObjectField(
+								*Schema_Properties, Last(MemberPath), *this);
+							Schema_Properties_TestBool.IsSet())
+						{
+							{
+								UE_JSON_SCHEMA_SCOPED_MEMBER_PATH_PUSH(MemberPath, TEXT("type"), true);
+								
+								if (const TOptional<FString> Schema_Properties_TestBool_Type = TestAndGetJsonStringField(
+										*Schema_Properties_TestBool, Last(MemberPath), *this);
+									Schema_Properties_TestBool_Type.IsSet())
+								{
+									TestTrue(
+										FString::Printf(TEXT("%s is boolean."), *MemberPath),
+										*Schema_Properties_TestBool_Type == "boolean"); 
+								}
+							}
+							
+							{
+								UE_JSON_SCHEMA_SCOPED_MEMBER_PATH_PUSH(MemberPath, TEXT("description"), true);
+								
+								if (const TOptional<FString> Schema_Properties_TestBool_Description = TestAndGetJsonStringField(
+										*Schema_Properties_TestBool, Last(MemberPath), *this);
+									Schema_Properties_TestBool_Description.IsSet())
+								{
+									TestEqual(
+										FString::Printf(TEXT("%s should be '%s'. Is '%s'."), *MemberPath, 
+											*ExpectedTestBoolDescription, **Schema_Properties_TestBool_Description),
+										*Schema_Properties_TestBool_Description, 
+										ExpectedTestBoolDescription); 
+								}
+							}
+						}
+					}
+				});
+			
+			It("Should generate JSON schema from a UStruct's string UProperty.", [this]() -> void
+				{
+					FJsonSchemaMemberPath MemberPath(TEXT("[schema]")); // ..mosty for error reporting
+					
+					const TSharedPtr<FJsonObject> Schema = FJsonSchemaGenerator::UStructToJsonSchemaObject<FJsonUtilitiesFakeStruct>();
+					if (!TestJsonSchemaIsValidAndNotEmpty(Schema, *this))
+					{
+						return;
+					}
+						
+					UE_JSON_SCHEMA_SCOPED_MEMBER_PATH_PUSH(MemberPath, TEXT("properties"), true);
+					
+					if (const TOptional<TSharedPtr<FJsonObject>> Schema_Properties = TestAndGetJsonObjectField(
+							Schema, Last(MemberPath), *this);
+						Schema_Properties.IsSet())
+					{
+						UE_JSON_SCHEMA_SCOPED_MEMBER_PATH_PUSH(MemberPath, TEXT("testString"), true);
+						
+						if (const TOptional<TSharedPtr<FJsonObject>> Schema_Properties_TestString = TestAndGetJsonObjectField(
+								*Schema_Properties, Last(MemberPath), *this);
+							Schema_Properties_TestString.IsSet())
+						{
+							{
+								UE_JSON_SCHEMA_SCOPED_MEMBER_PATH_PUSH(MemberPath, TEXT("type"), true);
+								
+								if (const TOptional<FString> Schema_Properties_TestString_Type = TestAndGetJsonStringField(
+										*Schema_Properties_TestString, Last(MemberPath), *this);
+									Schema_Properties_TestString_Type.IsSet())
+								{
+									TestTrue(
+										FString::Printf(TEXT("%s is string."), *MemberPath),
+										*Schema_Properties_TestString_Type == "string");
+								}
+							}
+
+							{
+								UE_JSON_SCHEMA_SCOPED_MEMBER_PATH_PUSH(MemberPath, TEXT("description"), true);
+								
+								if (const TOptional<FString> Schema_Properties_TestString_Description = TestAndGetJsonStringField(
+										*Schema_Properties_TestString, Last(MemberPath), *this);
+									Schema_Properties_TestString_Description.IsSet())
+								{
+									TestEqual(
+										FString::Printf(TEXT("%s should be '%s'. Is '%s'."), *MemberPath, 
+											*ExpectedTestStringDescription, **Schema_Properties_TestString_Description),
+										*Schema_Properties_TestString_Description, 
+										ExpectedTestStringDescription); 
+								}
+							}
+						}
+					}
+				});
+			
+			It("Should generate JSON schema from a UStruct's enum UProperty.", [this]() -> void
+				{
+					FJsonSchemaMemberPath MemberPath(TEXT("[schema]")); // ..mosty for error reporting
+					
+					const TSharedPtr<FJsonObject> Schema = FJsonSchemaGenerator::UStructToJsonSchemaObject<FJsonUtilitiesFakeStruct>();
+					if (!TestJsonSchemaIsValidAndNotEmpty(Schema, *this))
+					{
+						return;
+					}
+						
+					UE_JSON_SCHEMA_SCOPED_MEMBER_PATH_PUSH(MemberPath, TEXT("properties"), true);
+					
+					if (const TOptional<TSharedPtr<FJsonObject>> Schema_Properties = TestAndGetJsonObjectField(
+							Schema, Last(MemberPath), *this);
+						Schema_Properties.IsSet())
+					{
+						UE_JSON_SCHEMA_SCOPED_MEMBER_PATH_PUSH(MemberPath, TEXT("testEnum"), true);
+						
+						if (const TOptional<TSharedPtr<FJsonObject>> Schema_Properties_TestEnum = TestAndGetJsonObjectField(
+								*Schema_Properties, Last(MemberPath), *this);
+							Schema_Properties_TestEnum.IsSet())
+						{
+							{
+								UE_JSON_SCHEMA_SCOPED_MEMBER_PATH_PUSH(MemberPath, TEXT("type"), true);
+								
+								if (const TOptional<FString> Schema_Properties_TestEnum_Type = TestAndGetJsonStringField(
+										*Schema_Properties_TestEnum, Last(MemberPath), *this);
+									Schema_Properties_TestEnum_Type.IsSet())
+								{
+									TestTrue(
+										FString::Printf(TEXT("%s is string."), *MemberPath),
+										*Schema_Properties_TestEnum_Type == "string"); 
+								}
+							}
+
+							{
+								UE_JSON_SCHEMA_SCOPED_MEMBER_PATH_PUSH(MemberPath, TEXT("enum"), true);
+								
+								if (const TOptional<TArray<FString>> Schema_Properties_TestEnum_Enum = TestAndGetJsonStringArrayField(
+										*Schema_Properties_TestEnum, Last(MemberPath), *this);
+									Schema_Properties_TestEnum_Enum.IsSet())
+								{
+									if (TestEqual(
+									FString::Printf(
+										TEXT("%s should have %d values. Has %d values."), 
+											*MemberPath,
+											ExpectedTestEnumValues.Num(),
+											Schema_Properties_TestEnum_Enum->Num()),
+										Schema_Properties_TestEnum_Enum->Num(), ExpectedTestEnumValues.Num())) 
+									{
+										TestTrue(
+											FString::Printf(
+												TEXT("%s should have values ['%s', '%s']. Has values ['%s', '%s']."), 
+												*MemberPath,
+												*ExpectedTestEnumValues[0], *ExpectedTestEnumValues[1],
+												*(*Schema_Properties_TestEnum_Enum)[0], *(*Schema_Properties_TestEnum_Enum)[1]),
+                                     		*(*Schema_Properties_TestEnum_Enum)[0] == ExpectedTestEnumValues[0] && 
+                                     			*(*Schema_Properties_TestEnum_Enum)[1] == ExpectedTestEnumValues[1]);
+									}
+								}
+							}
+
+							{
+								UE_JSON_SCHEMA_SCOPED_MEMBER_PATH_PUSH(MemberPath, TEXT("description"), true);
+								
+								if (const TOptional<FString> Schema_Properties_TestEnum_Description = TestAndGetJsonStringField(
+										*Schema_Properties_TestEnum, Last(MemberPath), *this);
+									Schema_Properties_TestEnum_Description.IsSet())
+								{
+									TestEqual(
+										FString::Printf(TEXT("%s should be '%s'. Is '%s'."), *MemberPath, 
+											*ExpectedTestEnumDescription, **Schema_Properties_TestEnum_Description),
+										*Schema_Properties_TestEnum_Description, 
+										ExpectedTestEnumDescription); 
+								}
+							}
+						}
+					}
+				});
+			
+			It("Should generate JSON schema from a UStruct's nested UStruct UProperty.", [this]() -> void
+				{
+					FJsonSchemaMemberPath MemberPath(TEXT("[schema]")); // ..mosty for error reporting
+					
+					const TSharedPtr<FJsonObject> Schema = FJsonSchemaGenerator::UStructToJsonSchemaObject<FJsonUtilitiesFakeStruct>();
+					if (!TestJsonSchemaIsValidAndNotEmpty(Schema, *this))
+					{
+						return;
+					}
+						
+					UE_JSON_SCHEMA_SCOPED_MEMBER_PATH_PUSH(MemberPath, TEXT("properties"), true);
+					
+					if (const TOptional<TSharedPtr<FJsonObject>> Schema_Properties = TestAndGetJsonObjectField(
+							Schema, Last(MemberPath), *this);
+						Schema_Properties.IsSet())
+					{
+						UE_JSON_SCHEMA_SCOPED_MEMBER_PATH_PUSH(MemberPath, TEXT("testNestedStructA"), true);
+						
+						if (const TOptional<TSharedPtr<FJsonObject>> Schema_Properties_TestStruct = TestAndGetJsonObjectField(
+								*Schema_Properties, Last(MemberPath), *this);
+							Schema_Properties_TestStruct.IsSet())
+						{
+							{
+								UE_JSON_SCHEMA_SCOPED_MEMBER_PATH_PUSH(MemberPath, TEXT("type"), true);
+								
+								if (const TOptional<FString> Schema_Properties_TestStruct_Type = TestAndGetJsonStringField(
+										*Schema_Properties_TestStruct, Last(MemberPath), *this);
+									Schema_Properties_TestStruct_Type.IsSet())
+								{
+									TestTrue(
+										FString::Printf(TEXT("%s is object."), *MemberPath),
+										*Schema_Properties_TestStruct_Type == "object"); 
+								}
+							}
+							
+							{
+								UE_JSON_SCHEMA_SCOPED_MEMBER_PATH_PUSH(MemberPath, TEXT("description"), true);
+								
+								if (const TOptional<FString> Schema_Properties_TestStruct_Description = TestAndGetJsonStringField(
+										*Schema_Properties_TestStruct, Last(MemberPath), *this);
+									Schema_Properties_TestStruct_Description.IsSet())
+								{
+									TestEqual(
+										FString::Printf(TEXT("%s should be '%s'. Is '%s'."), *MemberPath, 
+											*ExpectedTestNestedStructADescription, **Schema_Properties_TestStruct_Description),
+										*Schema_Properties_TestStruct_Description, 
+										ExpectedTestNestedStructADescription); 
+								}
+							}
+						}
+					}
+				});
+
+			It("Should generate JSON schema from a UClass's UFunction, including its parameters.", [this]() -> void
+				{
+					FJsonSchemaMemberPath MemberPath(TEXT("[schema]")); // ..mosty for error reporting
+					
+					const UClass& Class = *UJsonUtilitiesFakeClass::StaticClass();
+
+					static const TCHAR* FunctionName = TEXT("TestFunc");
+					UFunction* FunctionPtr = Class.FindFunctionByName(FunctionName);
+					if (!TestTrue(
+						"Found named function.", 
+						LIKELY(FunctionPtr)))
+					{
+						return;
+					}
+					
+					const TSharedPtr<FJsonObject> Schema = FJsonSchemaGenerator::UStructToJsonSchemaObject(FunctionPtr);
+					if (!TestJsonSchemaIsValidAndNotEmpty(Schema, *this))
+					{
+						return;
+					}
+
+					MaybePrintJsonToLog(FunctionName, Schema.ToSharedRef());
+
+					{
+						UE_JSON_SCHEMA_SCOPED_MEMBER_PATH_PUSH(MemberPath, TEXT("inputSchema"), true);
+						
+						if (const TOptional<TSharedPtr<FJsonObject>> Schema_InputSchema = TestAndGetJsonObjectField(
+								Schema, Last(MemberPath), *this);
+							Schema_InputSchema.IsSet())
+						{
+							UE_JSON_SCHEMA_SCOPED_MEMBER_PATH_PUSH(MemberPath, TEXT("properties"), true);
+							
+							if (const TOptional<TSharedPtr<FJsonObject>> Schema_InputSchema_Properties = TestAndGetJsonObjectField(
+									*Schema_InputSchema, Last(MemberPath), *this);
+								Schema_InputSchema_Properties.IsSet())
+							{
+								UE_JSON_SCHEMA_SCOPED_MEMBER_PATH_PUSH(MemberPath, TEXT("testFuncParamFloat"), true);
+								
+								if (const TOptional<TSharedPtr<FJsonObject>> Schema_InputSchema_Properties_TestFuncParamFloat = TestAndGetJsonObjectField(
+										*Schema_InputSchema_Properties, Last(MemberPath), *this);
+									Schema_InputSchema_Properties_TestFuncParamFloat.IsSet())
+								{
+									{
+										UE_JSON_SCHEMA_SCOPED_MEMBER_PATH_PUSH(MemberPath, TEXT("type"), true);
+										
+										if (const TOptional<FString> Schema_InputSchema_Properties_TestFuncParamFloat_Type = TestAndGetJsonStringField(
+												*Schema_InputSchema_Properties_TestFuncParamFloat, Last(MemberPath), *this);
+											Schema_InputSchema_Properties_TestFuncParamFloat_Type.IsSet())
+										{
+											TestTrue(
+												FString::Printf(TEXT("%s is number."), *MemberPath),
+												*Schema_InputSchema_Properties_TestFuncParamFloat_Type == "number"); 
+										}
+									}
+
+									{
+										UE_JSON_SCHEMA_SCOPED_MEMBER_PATH_PUSH(MemberPath, TEXT("description"), true);
+										
+										if (const TOptional<FString> Schema_InputSchema_Properties_TestFuncParamFloat_Description = TestAndGetJsonStringField(
+												*Schema_InputSchema_Properties_TestFuncParamFloat, Last(MemberPath), *this);
+											Schema_InputSchema_Properties_TestFuncParamFloat_Description.IsSet())
+										{
+											TestEqual(
+												FString::Printf(TEXT("%s should be '%s'. Is '%s'."), *MemberPath, 
+													*ExpectedTestFuncParamFloatDescription, **Schema_InputSchema_Properties_TestFuncParamFloat_Description),
+												*Schema_InputSchema_Properties_TestFuncParamFloat_Description, 
+												ExpectedTestFuncParamFloatDescription); 
+										}
+									}
+								}
+							}
+						}
+					}
+					
+					{
+						UE_JSON_SCHEMA_SCOPED_MEMBER_PATH_PUSH(MemberPath, TEXT("outputSchema"), true);
+						
+						if (const TOptional<TSharedPtr<FJsonObject>> Schema_OutputSchema = TestAndGetJsonObjectField(
+								Schema, Last(MemberPath), *this);
+							Schema_OutputSchema.IsSet())
+						{
+							UE_JSON_SCHEMA_SCOPED_MEMBER_PATH_PUSH(MemberPath, TEXT("properties"), true);
+							
+							if (const TOptional<TSharedPtr<FJsonObject>> Schema_OutputSchema_Properties = TestAndGetJsonObjectField(
+									*Schema_OutputSchema, Last(MemberPath), *this);
+								Schema_OutputSchema_Properties.IsSet())
+							{
+								UE_JSON_SCHEMA_SCOPED_MEMBER_PATH_PUSH(MemberPath, TEXT("returnValue"), true);
+								
+								if (const TOptional<TSharedPtr<FJsonObject>> Schema_OutputSchema_Properties_TestReturnValueObject = TestAndGetJsonObjectField(
+										*Schema_OutputSchema_Properties, Last(MemberPath), *this);
+									Schema_OutputSchema_Properties_TestReturnValueObject.IsSet())
+								{
+									{
+										UE_JSON_SCHEMA_SCOPED_MEMBER_PATH_PUSH(MemberPath, TEXT("type"), true);
+										
+										if (TOptional<FString> Schema_OutputSchema_Properties_TestFuncReturnValue_Type = TestAndGetJsonStringField(
+												*Schema_OutputSchema_Properties_TestReturnValueObject, Last(MemberPath), *this);
+											Schema_OutputSchema_Properties_TestFuncReturnValue_Type.IsSet())
+										{
+											TestTrue(
+												FString::Printf(TEXT("%s type is number."), *MemberPath),
+												*Schema_OutputSchema_Properties_TestFuncReturnValue_Type == "number"); 
+										}
+									}
+
+									{
+										UE_JSON_SCHEMA_SCOPED_MEMBER_PATH_PUSH(MemberPath, TEXT("description"), true);
+										
+										if (const TOptional<FString> Schema_OutputSchema_Properties_TestFuncReturnValue_Description = TestAndGetJsonStringField(
+												*Schema_OutputSchema_Properties_TestReturnValueObject, Last(MemberPath), *this);
+											Schema_OutputSchema_Properties_TestFuncReturnValue_Description.IsSet())
+										{
+											TestEqual(
+												FString::Printf(TEXT("%s should be '%s'. Is '%s'."), *MemberPath, 
+													*ExpectedTestFuncReturnValueDescription, **Schema_OutputSchema_Properties_TestFuncReturnValue_Description),
+												*Schema_OutputSchema_Properties_TestFuncReturnValue_Description, 
+												ExpectedTestFuncReturnValueDescription); 
+										}
+									}
+								}
+							}
+						}
+					}
+
+					{
+						UE_JSON_SCHEMA_SCOPED_MEMBER_PATH_PUSH(MemberPath, TEXT("description"), true);
+						
+						// In this case, overall schema description should be function description.
+						if (const TOptional<FString> Schema_Description = TestAndGetJsonStringField(
+								Schema, Last(MemberPath), *this);
+							Schema_Description.IsSet())
+						{
+							TestEqual(
+								FString::Printf(TEXT("%s should be '%s'. Is '%s'."), *MemberPath, 
+									*ExpectedTestFuncDescription, **Schema_Description),
+								*Schema_Description, 
+								ExpectedTestFuncDescription); 
+						}
+					}
+				});
+			
+		});
+}
+
+// See comment above.
+PRAGMA_ENABLE_DEPRECATION_WARNINGS
+
+#endif
diff --git a/Engine/Plugins/Tests/EditorTests/Source/EditorTests/Private/UnrealEd/JsonUtilitiesTests.h b/Engine/Plugins/Tests/EditorTests/Source/EditorTests/Private/UnrealEd/JsonUtilitiesTests.h
new file mode 100644
--- /dev/null
+++ b/Engine/Plugins/Tests/EditorTests/Source/EditorTests/Private/UnrealEd/JsonUtilitiesTests.h
@@ -0,0 +1,186 @@
+// Copyright Epic Games, Inc. All Rights Reserved. 
+
+
+#pragma once
+
+
+#include "CoreMinimal.h"
+
+#include "JsonUtilitiesTests.generated.h"
+
+
+static const FString ExpectedTestFloatDescription = FString(TEXT("TestFloatDesc"));
+static constexpr float ExpectedTestFloatMinimum = -101.0f;
+static constexpr float ExpectedTestFloatMaximum = 102.0f;
+
+static const FString ExpectedTestIntDescription = FString(TEXT("TestIntDesc"));
+static constexpr int32 ExpectedTestIntMinimum = -201;
+static constexpr int32 ExpectedTestIntMaximum = 202;
+
+static const FString ExpectedTestBoolDescription = FString(TEXT("TestBoolDesc"));
+
+static const FString ExpectedTestEnumDescription = FString(TEXT("TestEnumDesc"));
+static const TArray<FString> ExpectedTestEnumValues = {TEXT("Zero"), TEXT("One")};
+
+static const FString ExpectedTestStringDescription = FString(TEXT("TestStringDesc"));
+
+static const FString ExpectedTestNestedStructADescription = FString(TEXT("TestNestedStructADesc"));
+static const FString ExpectedTestNestedStructBDescription = FString(TEXT("TestNestedStructBDesc"));
+static const FString ExpectedTestNestedStructCDescription = FString(TEXT("TestNestedStructCDesc"));
+static const FString ExpectedTestNestedStructDDescription = FString(TEXT("TestNestedStructDDesc"));
+
+static const FString ExpectedTestFuncDescription = FString(TEXT("TestFuncDesc"));
+static const FString ExpectedTestFuncWorldContext = FString(TEXT("TestFuncParamWorldContextObject"));
+static const FString ExpectedTestFuncParamFloatDescription = FString(TEXT("TestFuncParamFloatDesc"));
+static const FString ExpectedTestFuncParamWorldContextObjectDescription = FString(TEXT("TestFuncParamWorldContextObjectDesc"));
+static const FString ExpectedTestFuncReturnValueDescription = FString(TEXT("TestFuncReturnValueDesc"));
+
+// NOTE_JSON_UTILITIES_EDITOR_TEST_BLUEPRINT_DEFAULTS
+// Default values for function parameters in the metadata will only be available if a class is
+// BlueprintType and a function is a Blueprint function.
+
+UENUM()
+enum class EJsonUtilitiesFakeEnum : uint8
+{
+	Zero = 0,
+	One = 1
+};
+
+/**
+ * TestNestedStructADesc
+ */
+USTRUCT() 
+struct FJsonUtilitiesFakeNestedStructA
+{
+	GENERATED_BODY()
+	
+	/**
+	 * TestNestedFloatDesc
+	 */ 
+	UPROPERTY(EditAnywhere, Category="Nested", meta=(ClampMin="-301", ClampMax="302"))
+	float TestNestedStructAFloat = 0.0f;
+};
+
+/**
+ * TestNestedStructBDesc
+ */
+USTRUCT() 
+struct FJsonUtilitiesFakeNestedStructB
+{
+	GENERATED_BODY()
+	
+	/**
+	 * TestNestedFloatDesc
+	 */ 
+	UPROPERTY(EditAnywhere, Category="Nested", meta=(ClampMin="-401", ClampMax="402"))
+	float TestNestedStructBFloat = 0.0f;
+};
+
+/**
+ * TestNestedStructCDesc
+ */
+USTRUCT() 
+struct FJsonUtilitiesFakeNestedStructC
+{
+	GENERATED_BODY()
+	
+	/**
+	 * TestNestedFloatDesc
+	 */ 
+	UPROPERTY(EditAnywhere, Category="Nested", meta=(ClampMin="-501", ClampMax="502"))
+	float TestNestedStructCFloat = 0.0f;
+};
+
+/**
+ * TestNestedStructDDesc
+ */
+USTRUCT() 
+struct FJsonUtilitiesFakeNestedStructD
+{
+	GENERATED_BODY()
+	
+	/**
+	 * TestNestedStructBDesc
+	 */ 
+	UPROPERTY(EditAnywhere, Category="Nested")
+	FJsonUtilitiesFakeNestedStructB TestNestedStructB;
+	
+	/**
+	 * TestNestedStructCDesc
+	 */ 
+	UPROPERTY(EditAnywhere, Category="Nested")
+	FJsonUtilitiesFakeNestedStructC TestNestedStructC;
+};
+
+/**
+ * TestStructDesc
+ */
+USTRUCT(BlueprintType)
+struct FJsonUtilitiesFakeStruct
+{
+	GENERATED_BODY()
+	
+	/**
+	 * TestFloatDesc
+	 */ 
+	UPROPERTY(EditAnywhere, Category="Main", meta=(ClampMin="-101.0", ClampMax="102.0"))
+	float TestFloat = 0.0f;
+	
+	/**
+	 * TestIntDesc
+	 */ 
+	UPROPERTY(EditAnywhere, Category="Main", meta=(ClampMin="-201", ClampMax="202"))
+	int32 TestInt = 0;
+
+	/**
+	 * TestBoolDesc
+	 */ 
+	UPROPERTY(EditAnywhere, Category="Main")
+	bool TestBool = false;
+
+	/**
+	 * TestStringDesc
+	 */ 
+	UPROPERTY(EditAnywhere, Category="Main")
+	FString TestString = "";
+	
+	/**
+	 * TestEnumDesc
+	 */ 
+	UPROPERTY(EditAnywhere, Category="Main")
+	EJsonUtilitiesFakeEnum TestEnum = EJsonUtilitiesFakeEnum::Zero;
+	
+	/**
+	 * TestNestedStructDDesc
+	 */ 
+	UPROPERTY(EditAnywhere, Category="Main")
+	FJsonUtilitiesFakeNestedStructD TestNestedStructD;
+	
+	/**
+	 * TestNestedStructADesc
+	 */ 
+	UPROPERTY(EditAnywhere, Category="Main")
+	FJsonUtilitiesFakeNestedStructA TestNestedStructA;
+};
+
+/**
+ * TestClassDesc
+ */
+UCLASS(BlueprintType)
+class UJsonUtilitiesFakeClass : public UObject, public FJsonUtilitiesFakeStruct
+{
+	GENERATED_BODY()
+	
+	/**
+	 * TestFuncDesc
+	 * @param TestFuncParamWorldContextObject TestFuncParamWorldContextObjectDesc
+	 * @param TestFuncParamFloat TestFuncParamFloatDesc
+	 * @return TestFuncReturnValueDesc
+	 */
+	UFUNCTION(BlueprintCallable, Category="Main", meta=(WorldContext="TestFuncParamWorldContextObject")) 
+	float TestFunc(const float TestFuncParamFloat, const FJsonUtilitiesFakeStruct FakeStruct, const UObject* TestFuncParamWorldContextObject = nullptr)
+	{
+		return 0.0f;
+	}
+};
+
