Json Schema Tests.

Trivial fix -
Moved static expected values for tests from .h to .cpp, to avoid warnings about them ending up becoming multip0le copies in each translation unit that would include the .h.

#Synced-CL 49093340

diff --git a/Engine/Plugins/Tests/EditorTests/Source/EditorTests/Private/UnrealEd/JsonUtilitiesTests.cpp b/Engine/Plugins/Tests/EditorTests/Source/EditorTests/Private/UnrealEd/JsonUtilitiesTests.cpp
--- a/Engine/Plugins/Tests/EditorTests/Source/EditorTests/Private/UnrealEd/JsonUtilitiesTests.cpp
+++ b/Engine/Plugins/Tests/EditorTests/Source/EditorTests/Private/UnrealEd/JsonUtilitiesTests.cpp
@@ -11,6 +11,38 @@
 
 #include "JsonSchema/JsonSchemaGenerator.h"
 
+
+// Expected values.
+
+static const FString ExpectedTestFloatDescription = FString(TEXT("TestFloatDesc"));
+static constexpr float ExpectedTestFloatMinimum = -101.0f;
+static constexpr float ExpectedTestFloatMaximum = 102.0f;
+
+static const FString ExpectedTestIntDescription = FString(TEXT("TestIntDesc"));
+static constexpr int32 ExpectedTestIntMinimum = -201;
+static constexpr int32 ExpectedTestIntMaximum = 202;
+
+static const FString ExpectedTestBoolDescription = FString(TEXT("TestBoolDesc"));
+
+static const FString ExpectedTestEnumDescription = FString(TEXT("TestEnumDesc"));
+static const TArray<FString> ExpectedTestEnumValues = {TEXT("Zero"), TEXT("One")};
+
+static const FString ExpectedTestStringDescription = FString(TEXT("TestStringDesc"));
+
+static const FString ExpectedTestNestedStructADescription = FString(TEXT("TestNestedStructADesc"));
+static const FString ExpectedTestNestedStructBDescription = FString(TEXT("TestNestedStructBDesc"));
+static const FString ExpectedTestNestedStructCDescription = FString(TEXT("TestNestedStructCDesc"));
+static const FString ExpectedTestNestedStructDDescription = FString(TEXT("TestNestedStructDDesc"));
+
+static const FString ExpectedTestFuncDescription = FString(TEXT("TestFuncDesc"));
+static const FString ExpectedTestFuncWorldContext = FString(TEXT("TestFuncParamWorldContextObject"));
+static const FString ExpectedTestFuncParamFloatDescription = FString(TEXT("TestFuncParamFloatDesc"));
+static const FString ExpectedTestFuncParamWorldContextObjectDescription = FString(TEXT("TestFuncParamWorldContextObjectDesc"));
+static const FString ExpectedTestFuncReturnValueDescription = FString(TEXT("TestFuncReturnValueDesc"));
+
+
+// Tests.
+
 #if WITH_DEV_AUTOMATION_TESTS
 
 // JsonSchema functions are currently marker 'deprecated' as they are still 'experimental'.
diff --git a/Engine/Plugins/Tests/EditorTests/Source/EditorTests/Private/UnrealEd/JsonUtilitiesTests.h b/Engine/Plugins/Tests/EditorTests/Source/EditorTests/Private/UnrealEd/JsonUtilitiesTests.h
--- a/Engine/Plugins/Tests/EditorTests/Source/EditorTests/Private/UnrealEd/JsonUtilitiesTests.h
+++ b/Engine/Plugins/Tests/EditorTests/Source/EditorTests/Private/UnrealEd/JsonUtilitiesTests.h
@@ -9,32 +9,6 @@
 #include "JsonUtilitiesTests.generated.h"
 
 
-static const FString ExpectedTestFloatDescription = FString(TEXT("TestFloatDesc"));
-static constexpr float ExpectedTestFloatMinimum = -101.0f;
-static constexpr float ExpectedTestFloatMaximum = 102.0f;
-
-static const FString ExpectedTestIntDescription = FString(TEXT("TestIntDesc"));
-static constexpr int32 ExpectedTestIntMinimum = -201;
-static constexpr int32 ExpectedTestIntMaximum = 202;
-
-static const FString ExpectedTestBoolDescription = FString(TEXT("TestBoolDesc"));
-
-static const FString ExpectedTestEnumDescription = FString(TEXT("TestEnumDesc"));
-static const TArray<FString> ExpectedTestEnumValues = {TEXT("Zero"), TEXT("One")};
-
-static const FString ExpectedTestStringDescription = FString(TEXT("TestStringDesc"));
-
-static const FString ExpectedTestNestedStructADescription = FString(TEXT("TestNestedStructADesc"));
-static const FString ExpectedTestNestedStructBDescription = FString(TEXT("TestNestedStructBDesc"));
-static const FString ExpectedTestNestedStructCDescription = FString(TEXT("TestNestedStructCDesc"));
-static const FString ExpectedTestNestedStructDDescription = FString(TEXT("TestNestedStructDDesc"));
-
-static const FString ExpectedTestFuncDescription = FString(TEXT("TestFuncDesc"));
-static const FString ExpectedTestFuncWorldContext = FString(TEXT("TestFuncParamWorldContextObject"));
-static const FString ExpectedTestFuncParamFloatDescription = FString(TEXT("TestFuncParamFloatDesc"));
-static const FString ExpectedTestFuncParamWorldContextObjectDescription = FString(TEXT("TestFuncParamWorldContextObjectDesc"));
-static const FString ExpectedTestFuncReturnValueDescription = FString(TEXT("TestFuncReturnValueDesc"));
-
 // NOTE_JSON_UTILITIES_EDITOR_TEST_BLUEPRINT_DEFAULTS
 // Default values for function parameters in the metadata will only be available if a class is
 // BlueprintType and a function is a Blueprint function.
