Json Schema Editor Tests -

Trivial changes -
- Mostly spelling corrections.
- Anonymous namespace for global statics.
- Passing a func param by const reference instead of non-const reference.

#Synced-CL 49096638

diff --git a/Engine/Plugins/Tests/EditorTests/Source/EditorTests/Private/UnrealEd/JsonUtilitiesTests.cpp b/Engine/Plugins/Tests/EditorTests/Source/EditorTests/Private/UnrealEd/JsonUtilitiesTests.cpp
--- a/Engine/Plugins/Tests/EditorTests/Source/EditorTests/Private/UnrealEd/JsonUtilitiesTests.cpp
+++ b/Engine/Plugins/Tests/EditorTests/Source/EditorTests/Private/UnrealEd/JsonUtilitiesTests.cpp
@@ -14,31 +14,34 @@
 
 // Expected values.
 
-static const FString ExpectedTestFloatDescription = FString(TEXT("TestFloatDesc"));
-static constexpr float ExpectedTestFloatMinimum = -101.0f;
-static constexpr float ExpectedTestFloatMaximum = 102.0f;
+namespace
+{
+	const FString ExpectedTestFloatDescription = FString(TEXT("TestFloatDesc"));
+	constexpr float ExpectedTestFloatMinimum = -101.0f;
+	constexpr float ExpectedTestFloatMaximum = 102.0f;
 
-static const FString ExpectedTestIntDescription = FString(TEXT("TestIntDesc"));
-static constexpr int32 ExpectedTestIntMinimum = -201;
-static constexpr int32 ExpectedTestIntMaximum = 202;
+	const FString ExpectedTestIntDescription = FString(TEXT("TestIntDesc"));
+	constexpr int32 ExpectedTestIntMinimum = -201;
+	constexpr int32 ExpectedTestIntMaximum = 202;
 
-static const FString ExpectedTestBoolDescription = FString(TEXT("TestBoolDesc"));
+	const FString ExpectedTestBoolDescription = FString(TEXT("TestBoolDesc"));
 
-static const FString ExpectedTestEnumDescription = FString(TEXT("TestEnumDesc"));
-static const TArray<FString> ExpectedTestEnumValues = {TEXT("Zero"), TEXT("One")};
+	const FString ExpectedTestEnumDescription = FString(TEXT("TestEnumDesc"));
+	const TArray<FString> ExpectedTestEnumValues = {TEXT("Zero"), TEXT("One")};
 
-static const FString ExpectedTestStringDescription = FString(TEXT("TestStringDesc"));
+	const FString ExpectedTestStringDescription = FString(TEXT("TestStringDesc"));
 
-static const FString ExpectedTestNestedStructADescription = FString(TEXT("TestNestedStructADesc"));
-static const FString ExpectedTestNestedStructBDescription = FString(TEXT("TestNestedStructBDesc"));
-static const FString ExpectedTestNestedStructCDescription = FString(TEXT("TestNestedStructCDesc"));
-static const FString ExpectedTestNestedStructDDescription = FString(TEXT("TestNestedStructDDesc"));
+	const FString ExpectedTestNestedStructADescription = FString(TEXT("TestNestedStructADesc"));
+	const FString ExpectedTestNestedStructBDescription = FString(TEXT("TestNestedStructBDesc"));
+	const FString ExpectedTestNestedStructCDescription = FString(TEXT("TestNestedStructCDesc"));
+	const FString ExpectedTestNestedStructDDescription = FString(TEXT("TestNestedStructDDesc"));
 
-static const FString ExpectedTestFuncDescription = FString(TEXT("TestFuncDesc"));
-static const FString ExpectedTestFuncWorldContext = FString(TEXT("TestFuncParamWorldContextObject"));
-static const FString ExpectedTestFuncParamFloatDescription = FString(TEXT("TestFuncParamFloatDesc"));
-static const FString ExpectedTestFuncParamWorldContextObjectDescription = FString(TEXT("TestFuncParamWorldContextObjectDesc"));
-static const FString ExpectedTestFuncReturnValueDescription = FString(TEXT("TestFuncReturnValueDesc"));
+	const FString ExpectedTestFuncDescription = FString(TEXT("TestFuncDesc"));
+	const FString ExpectedTestFuncWorldContext = FString(TEXT("TestFuncParamWorldContextObject"));
+	const FString ExpectedTestFuncParamFloatDescription = FString(TEXT("TestFuncParamFloatDesc"));
+	const FString ExpectedTestFuncParamWorldContextObjectDescription = FString(TEXT("TestFuncParamWorldContextObjectDesc"));
+	const FString ExpectedTestFuncReturnValueDescription = FString(TEXT("TestFuncReturnValueDesc"));
+}
 
 
 // Tests.
@@ -54,7 +57,7 @@ PRAGMA_DISABLE_DEPRECATION_WARNINGS
 namespace
 {
 	// Gets last member of the member path.
-	FString Last(FJsonSchemaMemberPath& MemberPath)
+	FString Last(const FJsonSchemaMemberPath& MemberPath)
 	{
 		const FString MemberPathString = MemberPath.ToString();
 	
@@ -164,11 +167,11 @@ namespace
 	}
 }
 
-BEGIN_DEFINE_SPEC(FJsonUtiltiesEditorSpec, "System.Promotion.Editor.JsonUtilities",
+BEGIN_DEFINE_SPEC(FJsonUtilitiesEditorSpec, "System.Promotion.Editor.JsonUtilities",
 	EAutomationTestFlags::EditorContext|EAutomationTestFlags::EngineFilter)
-END_DEFINE_SPEC(FJsonUtiltiesEditorSpec)
+END_DEFINE_SPEC(FJsonUtilitiesEditorSpec)
 
-void FJsonUtiltiesEditorSpec::Define()
+void FJsonUtilitiesEditorSpec::Define()
 {
 	BeforeEach([this]() -> void
 		{
@@ -206,7 +209,7 @@ void FJsonUtiltiesEditorSpec::Define()
 			
 			It("Should generate JSON schema from a UStruct's float UProperty.", [this]() -> void
 				{
-					FJsonSchemaMemberPath MemberPath(TEXT("[schema]")); // ..mosty for error reporting
+					FJsonSchemaMemberPath MemberPath(TEXT("[schema]")); // ..mostly for error reporting
 					
 					const TSharedPtr<FJsonObject> Schema = FJsonSchemaGenerator::UStructToJsonSchemaObject<FJsonUtilitiesFakeStruct>();
 					if (!TestJsonSchemaIsValidAndNotEmpty(Schema, *this))
@@ -290,7 +293,7 @@ void FJsonUtiltiesEditorSpec::Define()
 			
 			It("Should generate JSON schema from a UStruct's int UProperty.", [this]() -> void
 				{
-					FJsonSchemaMemberPath MemberPath(TEXT("[schema]")); // ..mosty for error reporting
+					FJsonSchemaMemberPath MemberPath(TEXT("[schema]")); // ..mostly for error reporting
 					
 					const TSharedPtr<FJsonObject> Schema = FJsonSchemaGenerator::UStructToJsonSchemaObject<FJsonUtilitiesFakeStruct>();
 					if (!TestJsonSchemaIsValidAndNotEmpty(Schema, *this))
@@ -373,7 +376,7 @@ void FJsonUtiltiesEditorSpec::Define()
 			
 			It("Should generate JSON schema from a UStruct's bool UProperty.", [this]() -> void
 				{
-					FJsonSchemaMemberPath MemberPath(TEXT("[schema]")); // ..mosty for error reporting
+					FJsonSchemaMemberPath MemberPath(TEXT("[schema]")); // ..mostly for error reporting
 					
 					const TSharedPtr<FJsonObject> Schema = FJsonSchemaGenerator::UStructToJsonSchemaObject<FJsonUtilitiesFakeStruct>();
 					if (!TestJsonSchemaIsValidAndNotEmpty(Schema, *this))
@@ -426,7 +429,7 @@ void FJsonUtiltiesEditorSpec::Define()
 			
 			It("Should generate JSON schema from a UStruct's string UProperty.", [this]() -> void
 				{
-					FJsonSchemaMemberPath MemberPath(TEXT("[schema]")); // ..mosty for error reporting
+					FJsonSchemaMemberPath MemberPath(TEXT("[schema]")); // ..mostly for error reporting
 					
 					const TSharedPtr<FJsonObject> Schema = FJsonSchemaGenerator::UStructToJsonSchemaObject<FJsonUtilitiesFakeStruct>();
 					if (!TestJsonSchemaIsValidAndNotEmpty(Schema, *this))
@@ -479,7 +482,7 @@ void FJsonUtiltiesEditorSpec::Define()
 			
 			It("Should generate JSON schema from a UStruct's enum UProperty.", [this]() -> void
 				{
-					FJsonSchemaMemberPath MemberPath(TEXT("[schema]")); // ..mosty for error reporting
+					FJsonSchemaMemberPath MemberPath(TEXT("[schema]")); // ..mostly for error reporting
 					
 					const TSharedPtr<FJsonObject> Schema = FJsonSchemaGenerator::UStructToJsonSchemaObject<FJsonUtilitiesFakeStruct>();
 					if (!TestJsonSchemaIsValidAndNotEmpty(Schema, *this))
@@ -559,7 +562,7 @@ void FJsonUtiltiesEditorSpec::Define()
 			
 			It("Should generate JSON schema from a UStruct's nested UStruct UProperty.", [this]() -> void
 				{
-					FJsonSchemaMemberPath MemberPath(TEXT("[schema]")); // ..mosty for error reporting
+					FJsonSchemaMemberPath MemberPath(TEXT("[schema]")); // ..mostly for error reporting
 					
 					const TSharedPtr<FJsonObject> Schema = FJsonSchemaGenerator::UStructToJsonSchemaObject<FJsonUtilitiesFakeStruct>();
 					if (!TestJsonSchemaIsValidAndNotEmpty(Schema, *this))
@@ -612,7 +615,7 @@ void FJsonUtiltiesEditorSpec::Define()
 
 			It("Should generate JSON schema from a UClass's UFunction, including its parameters.", [this]() -> void
 				{
-					FJsonSchemaMemberPath MemberPath(TEXT("[schema]")); // ..mosty for error reporting
+					FJsonSchemaMemberPath MemberPath(TEXT("[schema]")); // ..mostly for error reporting
 					
 					const UClass& Class = *UJsonUtilitiesFakeClass::StaticClass();
 
