Editor Tests / Json Utilties / Json Schema

- Trivial: Default values are supplied for POD types and strings in structs and functions, and some default values changed.

- Trivial: UJsonUtilitiesFakeClass::TestFunc() reordered parameters, to be able to supply a default value for the float parameter.

- Trivial: Fixed UJsonUtilitiesFakeClass::TestFunc() comment.

#Synced-CL 49302269

diff --git a/Engine/Plugins/Tests/EditorTests/Source/EditorTests/Private/UnrealEd/JsonUtilitiesTests.h b/Engine/Plugins/Tests/EditorTests/Source/EditorTests/Private/UnrealEd/JsonUtilitiesTests.h
--- a/Engine/Plugins/Tests/EditorTests/Source/EditorTests/Private/UnrealEd/JsonUtilitiesTests.h
+++ b/Engine/Plugins/Tests/EditorTests/Source/EditorTests/Private/UnrealEd/JsonUtilitiesTests.h
@@ -32,7 +32,7 @@ struct FJsonUtilitiesFakeNestedStructA
 	 * TestNestedFloatDesc
 	 */ 
 	UPROPERTY(EditAnywhere, Category="Nested", meta=(ClampMin="-301", ClampMax="302"))
-	float TestNestedStructAFloat = 0.0f;
+	float TestNestedStructAFloat = 3.0f;
 };
 
 /**
@@ -47,7 +47,7 @@ struct FJsonUtilitiesFakeNestedStructB
 	 * TestNestedFloatDesc
 	 */ 
 	UPROPERTY(EditAnywhere, Category="Nested", meta=(ClampMin="-401", ClampMax="402"))
-	float TestNestedStructBFloat = 0.0f;
+	float TestNestedStructBFloat = 4.0f;
 };
 
 /**
@@ -62,7 +62,7 @@ struct FJsonUtilitiesFakeNestedStructC
 	 * TestNestedFloatDesc
 	 */ 
 	UPROPERTY(EditAnywhere, Category="Nested", meta=(ClampMin="-501", ClampMax="502"))
-	float TestNestedStructCFloat = 0.0f;
+	float TestNestedStructCFloat = 5.0f;
 };
 
 /**
@@ -98,31 +98,31 @@ struct FJsonUtilitiesFakeStruct
 	 * TestFloatDesc
 	 */ 
 	UPROPERTY(EditAnywhere, Category="Main", meta=(ClampMin="-101.0", ClampMax="102.0"))
-	float TestFloat = 0.0f;
+	float TestFloat = 1.0f;
 	
 	/**
 	 * TestIntDesc
 	 */ 
 	UPROPERTY(EditAnywhere, Category="Main", meta=(ClampMin="-201", ClampMax="202"))
-	int32 TestInt = 0;
+	int32 TestInt = 2;
 
 	/**
 	 * TestBoolDesc
 	 */ 
 	UPROPERTY(EditAnywhere, Category="Main")
-	bool TestBool = false;
+	bool TestBool = true;
 
 	/**
 	 * TestStringDesc
 	 */ 
 	UPROPERTY(EditAnywhere, Category="Main")
-	FString TestString = "";
+	FString TestString = "DefaultString";
 	
 	/**
 	 * TestEnumDesc
 	 */ 
 	UPROPERTY(EditAnywhere, Category="Main")
-	EJsonUtilitiesFakeEnum TestEnum = EJsonUtilitiesFakeEnum::Zero;
+	EJsonUtilitiesFakeEnum TestEnum = EJsonUtilitiesFakeEnum::One;
 	
 	/**
 	 * TestNestedStructDDesc
@@ -147,12 +147,13 @@ class UJsonUtilitiesFakeClass : public UObject, public FJsonUtilitiesFakeStruct
 	
 	/**
 	 * TestFuncDesc
-	 * @param TestFuncParamWorldContextObject TestFuncParamWorldContextObjectDesc
+	 * @param TestFuncParamFakeStruct TestFuncParamFakeStructDesc
 	 * @param TestFuncParamFloat TestFuncParamFloatDesc
+	 * @param TestFuncParamWorldContextObject TestFuncParamWorldContextObjectDesc
 	 * @return TestFuncReturnValueDesc
 	 */
 	UFUNCTION(BlueprintCallable, Category="Main", meta=(WorldContext="TestFuncParamWorldContextObject")) 
-	float TestFunc(const float TestFuncParamFloat, const FJsonUtilitiesFakeStruct FakeStruct, const UObject* TestFuncParamWorldContextObject = nullptr)
+	float TestFunc(const FJsonUtilitiesFakeStruct TestFuncParamFakeStruct, const float TestFuncParamFloat = 1.23f, const UObject* TestFuncParamWorldContextObject = nullptr)
 	{
 		return 0.0f;
 	}
