Editor Tests, Json Utilities, Json Schema

- Trivial: UJsonUtilitiesFakeClass::TestFunc() now supplies more parameters, to check default values - for ints and enums.

#Synced-CL 49308818

diff --git a/Engine/Plugins/Tests/EditorTests/Source/EditorTests/Private/UnrealEd/JsonUtilitiesTests.h b/Engine/Plugins/Tests/EditorTests/Source/EditorTests/Private/UnrealEd/JsonUtilitiesTests.h
--- a/Engine/Plugins/Tests/EditorTests/Source/EditorTests/Private/UnrealEd/JsonUtilitiesTests.h
+++ b/Engine/Plugins/Tests/EditorTests/Source/EditorTests/Private/UnrealEd/JsonUtilitiesTests.h
@@ -149,11 +149,14 @@ class UJsonUtilitiesFakeClass : public UObject, public FJsonUtilitiesFakeStruct
 	 * TestFuncDesc
 	 * @param TestFuncParamFakeStruct TestFuncParamFakeStructDesc
 	 * @param TestFuncParamFloat TestFuncParamFloatDesc
+	 * @param TestFuncParamInt TestFuncParamIntDesc
+	 * @param TestFuncParamEnum TestFuncParamEnumDesc
 	 * @param TestFuncParamWorldContextObject TestFuncParamWorldContextObjectDesc
 	 * @return TestFuncReturnValueDesc
 	 */
 	UFUNCTION(BlueprintCallable, Category="Main", meta=(WorldContext="TestFuncParamWorldContextObject")) 
-	float TestFunc(const FJsonUtilitiesFakeStruct TestFuncParamFakeStruct, const float TestFuncParamFloat = 1.23f, const UObject* TestFuncParamWorldContextObject = nullptr)
+	float TestFunc(const FJsonUtilitiesFakeStruct TestFuncParamFakeStruct, const float TestFuncParamFloat = 1.23f, const int32 TestFuncParamInt = 123,
+		const EJsonUtilitiesFakeEnum TestFuncParamEnum = EJsonUtilitiesFakeEnum::One, const UObject* TestFuncParamWorldContextObject = nullptr)
 	{
 		return 0.0f;
 	}
