Editor Tests, Json Utiltiies, Json Schema

- Trivial: Added bool parameter with default to UJsonUtilitiesFakeClass::TextFunc()

#Synced-CL 49315116

diff --git a/Engine/Plugins/Tests/EditorTests/Source/EditorTests/Private/UnrealEd/JsonUtilitiesTests.h b/Engine/Plugins/Tests/EditorTests/Source/EditorTests/Private/UnrealEd/JsonUtilitiesTests.h
--- a/Engine/Plugins/Tests/EditorTests/Source/EditorTests/Private/UnrealEd/JsonUtilitiesTests.h
+++ b/Engine/Plugins/Tests/EditorTests/Source/EditorTests/Private/UnrealEd/JsonUtilitiesTests.h
@@ -148,6 +148,7 @@ class UJsonUtilitiesFakeClass : public UObject, public FJsonUtilitiesFakeStruct
 	/**
 	 * TestFuncDesc
 	 * @param TestFuncParamFakeStruct TestFuncParamFakeStructDesc
+	 * @param bTestFuncParamBool TestFuncParamBoolDesc
 	 * @param TestFuncParamFloat TestFuncParamFloatDesc
 	 * @param TestFuncParamInt TestFuncParamIntDesc
 	 * @param TestFuncParamEnum TestFuncParamEnumDesc
@@ -155,7 +156,8 @@ class UJsonUtilitiesFakeClass : public UObject, public FJsonUtilitiesFakeStruct
 	 * @return TestFuncReturnValueDesc
 	 */
 	UFUNCTION(BlueprintCallable, Category="Main", meta=(WorldContext="TestFuncParamWorldContextObject")) 
-	float TestFunc(const FJsonUtilitiesFakeStruct TestFuncParamFakeStruct, const float TestFuncParamFloat = 1.23f, const int32 TestFuncParamInt = 123,
+	float TestFunc(const FJsonUtilitiesFakeStruct TestFuncParamFakeStruct, const bool bTestFuncParamBool = true, 
+		const float TestFuncParamFloat = 1.23f, const int32 TestFuncParamInt = 123,
 		const EJsonUtilitiesFakeEnum TestFuncParamEnum = EJsonUtilitiesFakeEnum::One, const UObject* TestFuncParamWorldContextObject = nullptr)
 	{
 		return 0.0f;
