Editor Tests / Json Schema

- Added test for function with struct params that are defaulted to stuct values (INCOMPLETE.)
- Added test for function with params that are illegal for JSON schemas.
- Removed WordContext object from standard function test.

#Synced-CL 49751710

diff --git a/Engine/Plugins/Tests/EditorTests/Source/EditorTests/Private/UnrealEd/JsonUtilitiesTests.cpp b/Engine/Plugins/Tests/EditorTests/Source/EditorTests/Private/UnrealEd/JsonUtilitiesTests.cpp
--- a/Engine/Plugins/Tests/EditorTests/Source/EditorTests/Private/UnrealEd/JsonUtilitiesTests.cpp
+++ b/Engine/Plugins/Tests/EditorTests/Source/EditorTests/Private/UnrealEd/JsonUtilitiesTests.cpp
@@ -37,9 +37,7 @@ namespace
 	const FString ExpectedTestNestedStructDDescription = FString(TEXT("TestNestedStructDDesc"));
 
 	const FString ExpectedTestFuncDescription = FString(TEXT("TestFuncDesc"));
-	const FString ExpectedTestFuncWorldContext = FString(TEXT("TestFuncParamWorldContextObject"));
 	const FString ExpectedTestFuncParamFloatDescription = FString(TEXT("TestFuncParamFloatDesc"));
-	const FString ExpectedTestFuncParamWorldContextObjectDescription = FString(TEXT("TestFuncParamWorldContextObjectDesc"));
 	const FString ExpectedTestFuncReturnValueDescription = FString(TEXT("TestFuncReturnValueDesc"));
 }
 
@@ -52,12 +50,12 @@ namespace
 // Remove this later once they are no longer deprecated..
 PRAGMA_DISABLE_DEPRECATION_WARNINGS
 
-#define PRINT_JSONS_TO_LOG 1
+#define DEBUG_PRINT_JSONS_TO_LOG 0
 
 namespace
 {
 	// Gets last member of the member path.
-	FString Last(const FJsonSchemaMemberPath& MemberPath)
+	FString Last(FJsonSchemaMemberPath& MemberPath)
 	{
 		const FString MemberPathString = MemberPath.ToString();
 	
@@ -70,18 +68,20 @@ namespace
 		return MemberPathString;
 	}
 	
-	bool TestJsonSchemaIsValidAndNotEmpty(TSharedPtr<FJsonObject> Schema, FAutomationTestBase& Test)
+	bool TestIsJsonSchemaValidAndNotEmpty(const TSharedPtr<FJsonObject>& Schema, FAutomationTestBase& Test)
 	{
-		if (!Test.TestTrue("Schema is valid.", Schema.IsValid()) ||
-			!Test.TestTrue("Schema is not empty (has values.)", Schema->Values.Num() > 0))
-		{
-			return false;
-		}
-		
-		return true;
+		return 
+			Test.TestTrue("Schema is valid.", Schema.IsValid()) &&
+			Test.TestTrue("Schema is not empty (has values.)", !Schema->Values.IsEmpty());
 	}
 	
-	TOptional<TSharedPtr<FJsonObject>> TestAndGetJsonObjectField(TSharedPtr<FJsonObject> Schema, const FString& FieldName, FAutomationTestBase& Test)
+	bool TestIsJsonSchemaInvalid(const TSharedPtr<FJsonObject>& Schema, FAutomationTestBase& Test)
+	{
+		return Test.TestFalse("Schema is invalid.", Schema.IsValid());
+	}
+	
+	TOptional<TSharedPtr<FJsonObject>> TestAndGetJsonObjectField(
+		const TSharedPtr<FJsonObject>& Schema, const FString& FieldName, FAutomationTestBase& Test)
 	{
 		TOptional<TSharedPtr<FJsonObject>> OptionalJsonObject;
 		
@@ -100,7 +100,8 @@ namespace
 		return OptionalJsonObject;
 	}
 	
-	TOptional<FString> TestAndGetJsonStringField(TSharedPtr<FJsonObject> Schema, const FString& FieldName, FAutomationTestBase& Test)
+	TOptional<FString> TestAndGetJsonStringField(
+		const TSharedPtr<FJsonObject>& Schema, const FString& FieldName, FAutomationTestBase& Test)
 	{
 		TOptional<FString> OptionalJsonString;
 		if (FString JsonString;
@@ -114,7 +115,8 @@ namespace
 		return OptionalJsonString;
 	}
 	
-	TOptional<TArray<FString>> TestAndGetJsonStringArrayField(TSharedPtr<FJsonObject> Schema, const FString& FieldName, FAutomationTestBase& Test)
+	TOptional<TArray<FString>> TestAndGetJsonStringArrayField(
+		const TSharedPtr<FJsonObject>& Schema, const FString& FieldName, FAutomationTestBase& Test)
 	{
 		TOptional<TArray<FString>> OptionalJsonStringArray;
 		if (TArray<FString> JsonStringArray;
@@ -128,7 +130,8 @@ namespace
 		return OptionalJsonStringArray;
 	}
 	
-	TOptional<float> TestAndGetJsonFloatNumberField(TSharedPtr<FJsonObject> Schema, const FString& FieldName, FAutomationTestBase& Test)
+	TOptional<float> TestAndGetJsonFloatNumberField(
+		const TSharedPtr<FJsonObject>& Schema, const FString& FieldName, FAutomationTestBase& Test)
 	{
 		TOptional<float> OptionalJsonNumber;
 		if (float JsonNumber = 0.0f;
@@ -142,7 +145,8 @@ namespace
 		return OptionalJsonNumber;
 	}
 	
-	TOptional<int32> TestAndGetJsonIntNumberField(TSharedPtr<FJsonObject> Schema, const FString& FieldName, FAutomationTestBase& Test)
+	TOptional<int32> TestAndGetJsonIntNumberField(
+		const TSharedPtr<FJsonObject>& Schema, const FString& FieldName, FAutomationTestBase& Test)
 	{
 		TOptional<int32> OptionalJsonNumber;
 		if (int32 JsonNumber = 0;
@@ -156,14 +160,20 @@ namespace
 		return OptionalJsonNumber;
 	}
 	
-	void MaybePrintJsonToLog(const TCHAR* Name, const TSharedRef<FJsonObject>& JsonObject)
+	void PrintJsonToLog(
+		const TCHAR* Name, const TSharedPtr<FJsonObject>& JsonObject)
 	{
-#if PRINT_JSONS_TO_LOG
-		FString JsonString;
-		const TSharedRef<TJsonWriter<>> JsonWriter = TJsonWriterFactory<>::Create(&JsonString);
-		FJsonSerializer::Serialize(JsonObject, JsonWriter);
-		UE_LOG(LogJson, Warning, TEXT("[%s]\n%s"), Name, *JsonString);
-#endif
+		if (JsonObject)
+		{
+			FString JsonString;
+			const TSharedRef<TJsonWriter<>> JsonWriter = TJsonWriterFactory<>::Create(&JsonString);
+			FJsonSerializer::Serialize(JsonObject, JsonWriter);
+			UE_LOG(LogJson, Display, TEXT("Json \"%s\" is valid -\n%s"), Name, *JsonString);
+		}
+		else
+		{
+			UE_LOG(LogJson, Error, TEXT("Json \"%s\" is invalid."), Name);
+		}
 	}
 }
 
@@ -183,36 +193,40 @@ void FJsonUtilitiesEditorSpec::Define()
 				LIKELY(Module));
 		});
 	
-	Describe("Should generate JSON schemas from UClasses/UStructs/UProperties/UFunctions.", [this]() -> void
+	Describe(
+		"Should generate JSON schemas from UClasses/UStructs/UProperties/UFunctions.", 
+		[this]() -> void
 		{
-			It("Should generate JSON schema from a UStruct - smoke test.", [this]() -> void
+			It(
+				"Should generate valid JSON schema from a UStruct - smoke test.", 
+				[this]() -> void
 				{
 					const TSharedPtr<FJsonObject> Schema = FJsonSchemaGenerator::UStructToJsonSchemaObject<FJsonUtilitiesFakeStruct>();
-					if (!TestJsonSchemaIsValidAndNotEmpty(Schema, *this))
+					if (!TestIsJsonSchemaValidAndNotEmpty(Schema, *this))
 					{
 						return;
 					}
-
-					MaybePrintJsonToLog(*FJsonUtilitiesFakeStruct::StaticStruct()->GetName(), Schema.ToSharedRef());
 				});
 			
-			It("Should generate JSON schema from a UClass - smoke test.", [this]() -> void
+			It(
+				"Should generate valid JSON schema from a UClass - smoke test.", 
+				[this]() -> void
 				{
 					const TSharedPtr<FJsonObject> Schema = FJsonSchemaGenerator::UStructToJsonSchemaObject<UJsonUtilitiesFakeClass>();
-					if (!TestJsonSchemaIsValidAndNotEmpty(Schema, *this))
+					if (!TestIsJsonSchemaValidAndNotEmpty(Schema, *this))
 					{
 						return;
 					}
-
-					MaybePrintJsonToLog(*UJsonUtilitiesFakeClass::StaticClass()->GetName(), Schema.ToSharedRef());
 				});
 			
-			It("Should generate JSON schema from a UStruct's float UProperty.", [this]() -> void
+			It(
+				"Should generate valid JSON schema from a UStruct's float UProperty.", 
+				[this]() -> void
 				{
 					FJsonSchemaMemberPath MemberPath(TEXT("[schema]")); // ..mostly for error reporting
 					
 					const TSharedPtr<FJsonObject> Schema = FJsonSchemaGenerator::UStructToJsonSchemaObject<FJsonUtilitiesFakeStruct>();
-					if (!TestJsonSchemaIsValidAndNotEmpty(Schema, *this))
+					if (!TestIsJsonSchemaValidAndNotEmpty(Schema, *this))
 					{
 						return;
 					}
@@ -291,12 +305,14 @@ void FJsonUtilitiesEditorSpec::Define()
 					}
 				});
 			
-			It("Should generate JSON schema from a UStruct's int UProperty.", [this]() -> void
+			It(
+				"Should generate valid JSON schema from a UStruct's int UProperty.", 
+				[this]() -> void
 				{
 					FJsonSchemaMemberPath MemberPath(TEXT("[schema]")); // ..mostly for error reporting
 					
 					const TSharedPtr<FJsonObject> Schema = FJsonSchemaGenerator::UStructToJsonSchemaObject<FJsonUtilitiesFakeStruct>();
-					if (!TestJsonSchemaIsValidAndNotEmpty(Schema, *this))
+					if (!TestIsJsonSchemaValidAndNotEmpty(Schema, *this))
 					{
 						return;
 					}
@@ -374,12 +390,14 @@ void FJsonUtilitiesEditorSpec::Define()
 					}
 				});
 			
-			It("Should generate JSON schema from a UStruct's bool UProperty.", [this]() -> void
+			It(
+				"Should generate valid JSON schema from a UStruct's bool UProperty.", 
+				[this]() -> void
 				{
 					FJsonSchemaMemberPath MemberPath(TEXT("[schema]")); // ..mostly for error reporting
 					
 					const TSharedPtr<FJsonObject> Schema = FJsonSchemaGenerator::UStructToJsonSchemaObject<FJsonUtilitiesFakeStruct>();
-					if (!TestJsonSchemaIsValidAndNotEmpty(Schema, *this))
+					if (!TestIsJsonSchemaValidAndNotEmpty(Schema, *this))
 					{
 						return;
 					}
@@ -427,12 +445,14 @@ void FJsonUtilitiesEditorSpec::Define()
 					}
 				});
 			
-			It("Should generate JSON schema from a UStruct's string UProperty.", [this]() -> void
+			It(
+				"Should generate valid JSON schema from a UStruct's string UProperty.", 
+				[this]() -> void
 				{
 					FJsonSchemaMemberPath MemberPath(TEXT("[schema]")); // ..mostly for error reporting
 					
 					const TSharedPtr<FJsonObject> Schema = FJsonSchemaGenerator::UStructToJsonSchemaObject<FJsonUtilitiesFakeStruct>();
-					if (!TestJsonSchemaIsValidAndNotEmpty(Schema, *this))
+					if (!TestIsJsonSchemaValidAndNotEmpty(Schema, *this))
 					{
 						return;
 					}
@@ -480,12 +500,14 @@ void FJsonUtilitiesEditorSpec::Define()
 					}
 				});
 			
-			It("Should generate JSON schema from a UStruct's enum UProperty.", [this]() -> void
+			It(
+				"Should generate valid JSON schema from a UStruct's enum UProperty.", 
+				[this]() -> void
 				{
 					FJsonSchemaMemberPath MemberPath(TEXT("[schema]")); // ..mostly for error reporting
 					
 					const TSharedPtr<FJsonObject> Schema = FJsonSchemaGenerator::UStructToJsonSchemaObject<FJsonUtilitiesFakeStruct>();
-					if (!TestJsonSchemaIsValidAndNotEmpty(Schema, *this))
+					if (!TestIsJsonSchemaValidAndNotEmpty(Schema, *this))
 					{
 						return;
 					}
@@ -560,12 +582,14 @@ void FJsonUtilitiesEditorSpec::Define()
 					}
 				});
 			
-			It("Should generate JSON schema from a UStruct's nested UStruct UProperty.", [this]() -> void
+			It(
+				"Should generate valid JSON schema from a UStruct - with nested UStruct UProperty.", 
+				[this]() -> void
 				{
 					FJsonSchemaMemberPath MemberPath(TEXT("[schema]")); // ..mostly for error reporting
 					
 					const TSharedPtr<FJsonObject> Schema = FJsonSchemaGenerator::UStructToJsonSchemaObject<FJsonUtilitiesFakeStruct>();
-					if (!TestJsonSchemaIsValidAndNotEmpty(Schema, *this))
+					if (!TestIsJsonSchemaValidAndNotEmpty(Schema, *this))
 					{
 						return;
 					}
@@ -613,7 +637,9 @@ void FJsonUtilitiesEditorSpec::Define()
 					}
 				});
 
-			It("Should generate JSON schema from a UClass's UFunction, including its parameters.", [this]() -> void
+			It(
+				"Should generate valid JSON schema from a UClass's UFunction - with simpler defaulted params.", 
+				[this]() -> void
 				{
 					FJsonSchemaMemberPath MemberPath(TEXT("[schema]")); // ..mostly for error reporting
 					
@@ -629,12 +655,14 @@ void FJsonUtilitiesEditorSpec::Define()
 					}
 					
 					const TSharedPtr<FJsonObject> Schema = FJsonSchemaGenerator::UStructToJsonSchemaObject(FunctionPtr);
-					if (!TestJsonSchemaIsValidAndNotEmpty(Schema, *this))
+					if (!TestIsJsonSchemaValidAndNotEmpty(Schema, *this))
 					{
 						return;
 					}
 
-					MaybePrintJsonToLog(FunctionName, Schema.ToSharedRef());
+#if DEBUG_PRINT_JSONS_TO_LOG
+					PrintJsonToLog(FunctionName, Schema);
+#endif
 
 					{
 						UE_JSON_SCHEMA_SCOPED_MEMBER_PATH_PUSH(MemberPath, TEXT("inputSchema"), true);
@@ -754,8 +782,61 @@ void FJsonUtilitiesEditorSpec::Define()
 						}
 					}
 				});
-			
-		});
+
+			It(
+				"Should generate valid JSON schema from a UClass's UFunction - with Engine struct defaulted params.", 
+				[this]() -> void
+				{
+					FJsonSchemaMemberPath MemberPath(TEXT("[schema]")); // ..mostly for error reporting
+					
+					const UClass& Class = *UJsonUtilitiesFakeClass::StaticClass();
+
+					static const TCHAR* FunctionName = TEXT("TestEngineStructFunc");
+					UFunction* FunctionPtr = Class.FindFunctionByName(FunctionName);
+					if (!TestTrue(
+						"Found named function.", 
+						LIKELY(FunctionPtr)))
+					{
+						return;
+					}
+					
+					const TSharedPtr<FJsonObject> Schema = FJsonSchemaGenerator::UStructToJsonSchemaObject(FunctionPtr);
+					if (!TestIsJsonSchemaValidAndNotEmpty(Schema, *this))
+					{
+						return;
+					}
+
+#if DEBUG_PRINT_JSONS_TO_LOG
+					PrintJsonToLog(FunctionName, Schema);
+#endif
+						
+					// TODO - Test is incomplete.
+				});
+
+			It(
+				"Should generate invalid JSON schema from a UClass's UFunction - with illegal non-Json-able params.", 
+				[this]() -> void
+				{
+					FJsonSchemaMemberPath MemberPath(TEXT("[schema]")); // ..mostly for error reporting
+				
+					const UClass& Class = *UJsonUtilitiesFakeClass::StaticClass();
+
+					static const TCHAR* FunctionName = TEXT("TestIllegalFunc");
+					UFunction* FunctionPtr = Class.FindFunctionByName(FunctionName);
+					if (!TestTrue(
+						"Found named function.", 
+						LIKELY(FunctionPtr)))
+					{
+						return;
+					}
+				
+					AddExpectedError(TEXT(".*Illegal property type for Json.*"),
+						EAutomationExpectedErrorFlags::Contains, INDEX_NONE);
+						
+					const TSharedPtr<FJsonObject> Schema = FJsonSchemaGenerator::UStructToJsonSchemaObject(FunctionPtr);
+					TestIsJsonSchemaInvalid(Schema, *this);
+			});
+	});
 }
 
 // See comment above.
diff --git a/Engine/Plugins/Tests/EditorTests/Source/EditorTests/Private/UnrealEd/JsonUtilitiesTests.h b/Engine/Plugins/Tests/EditorTests/Source/EditorTests/Private/UnrealEd/JsonUtilitiesTests.h
--- a/Engine/Plugins/Tests/EditorTests/Source/EditorTests/Private/UnrealEd/JsonUtilitiesTests.h
+++ b/Engine/Plugins/Tests/EditorTests/Source/EditorTests/Private/UnrealEd/JsonUtilitiesTests.h
@@ -5,6 +5,10 @@
 
 
 #include "CoreMinimal.h"
+#include "UObject/Object.h"
+#include "UObject/SoftObjectPath.h"
+#include "Math/Color.h"
+#include "Math/MathFwd.h"
 
 #include "JsonUtilitiesTests.generated.h"
 
@@ -137,6 +141,22 @@ struct FJsonUtilitiesFakeStruct
 	FJsonUtilitiesFakeNestedStructA TestNestedStructA;
 };
 
+/**
+ * FFunctionParamStruct
+ * Used for testing how structs end up working with function parameters in function schemas.
+ */
+USTRUCT(BlueprintType)
+struct FFunctionParamStruct
+{
+	GENERATED_BODY()
+	
+	UPROPERTY(Category="Main", BlueprintReadWrite, EditAnywhere)
+	FString AString = "";
+
+	UPROPERTY(Category="Main", BlueprintReadWrite, EditAnywhere)
+	FString BString = "";
+};
+
 /**
  * TestClassDesc
  */
@@ -152,15 +172,77 @@ class UJsonUtilitiesFakeClass : public UObject, public FJsonUtilitiesFakeStruct
 	 * @param TestFuncParamFloat TestFuncParamFloatDesc
 	 * @param TestFuncParamInt TestFuncParamIntDesc
 	 * @param TestFuncParamEnum TestFuncParamEnumDesc
-	 * @param TestFuncParamWorldContextObject TestFuncParamWorldContextObjectDesc
 	 * @return TestFuncReturnValueDesc
 	 */
-	UFUNCTION(BlueprintCallable, Category="Main", meta=(WorldContext="TestFuncParamWorldContextObject")) 
+	UFUNCTION(BlueprintCallable, Category="Main") 
 	float TestFunc(const FJsonUtilitiesFakeStruct TestFuncParamFakeStruct, const bool bTestFuncParamBool = true, 
 		const float TestFuncParamFloat = 1.23f, const int32 TestFuncParamInt = 123,
-		const EJsonUtilitiesFakeEnum TestFuncParamEnum = EJsonUtilitiesFakeEnum::One, const UObject* TestFuncParamWorldContextObject = nullptr)
+		const EJsonUtilitiesFakeEnum TestFuncParamEnum = EJsonUtilitiesFakeEnum::One)
+	{
+		// This function is for testing schema generation for a function with 
+		// simpler parameters.
+		
+		return 0.0f;
+	}
+	
+	/**
+	 * TestEngineStructFuncDesc
+	 * @param SoftObjectPath TestEngineStructFuncParamSoftObjectPathDesc 
+	 * @param SoftClassPath TestEngineStructFuncParamSoftClassPathDesc 
+	 * @param FunctionParamStruct TestEngineStructFuncParamFunctionParamStructDesc 
+	 * @param LinearColor TestEngineStructFuncParamLinearColorDesc 
+	 * @param Color TestEngineStructFuncParamColorDesc 
+	 * @param Quat TestEngineStructFuncParamQuatDesc 
+	 * @param Rotator TestEngineStructFuncParamRotatorDesc 
+	 * @param Transform TestEngineStructFuncParamTransformDesc 
+	 * @param Vector TestEngineStructFuncParamVectorDesc
+	 * @return TestEngineStructFuncReturnValueDesc
+	 */
+	UFUNCTION(BlueprintCallable, Category="Main", meta=(
+		CPP_Default_SoftObjectPath="/Game/MyFolder/MyAsset.MyAsset", 
+		CPP_Default_SoftClassPath="/Game/MyFolder/MyActorBP.MyMyActorBP_C",
+		CPP_Default_FunctionParamStruct="(AString=\"AAA\",BString=\"BBB\")",
+		CPP_Default_LinearColor="(R=0.1,G=0.2,B=0.3,A=0.4)",
+		CPP_Default_Color="(R=10,G=20,B=30,A=40)",
+		CPP_Default_Quat="(X=0.1,Y=0.2,Z=0.3,W=0.4)",
+		CPP_Default_Rotator="(Roll=30,Yaw=10,Pitch=20)",
+		CPP_Default_Transform="(Rotation=(X=0.1,Y=0.2,Z=0.3,W=0.4),Translation=(X=1,Y=2,Z=3),Scale3D=(X=10,Y=20,Z=30))"))
+	float TestEngineStructFunc(
+		FSoftObjectPath SoftObjectPath, 
+		FSoftClassPath SoftClassPath,
+		FFunctionParamStruct FunctionParamStruct, 
+		FLinearColor LinearColor, FColor Color,
+		FQuat Quat, FRotator Rotator, FTransform Transform, 
+		FVector Vector = FVector::UpVector)
 	{
+		// This function is for testing schema generation for a function with
+		// more complex struct parameters, with defaults.
+		// UHT does allow some structs of be initialized in the declaration, like FVector.
+		// UHT does not allow most other structs to be initialized in the declaration, like 
+		// FSoftObjectPath, or arbitrary structs like FFunctionParamStruct.. and for those we 
+		// need to use CPP_Default_XXX. 
+		// The real purpose of this function is for testing how these defaults end up looking in
+		// the generated JSON schema.
+
 		return 0.0f;
 	}
+	
+	/**
+	 * TestIllegalFuncDesc
+	 * @return TestIllegalFuncReturnValueDesc
+	 */
+	UFUNCTION(BlueprintCallable, Category="Main")
+	UObject* TestIllegalFunc(
+		int32 IntParamA = 1, 
+		UObject* ObjectParam = nullptr, 
+		int32 IntParamB = 2, 
+		UClass* ClassParam = nullptr, 
+		int32 IntParamC = 3)
+	{
+		// This function is for testing schema generation for a function with some parameters that 
+		// can't be represented in the JSON schema.
+
+		return nullptr;
+	}
 };
 
