Improve JsonSchema generation for complex types.
* UObject* and UClass* are all supported by default because they can be converted to and from strings.
* Only FMaps with string keys are supported. Other keys are unsupported but no longer trip an ensure.
* Removed some unnecessary validation code that was out of sync with actual schema gen logic.
* Added code for custom property generation callbacks.
* Rewrote tests to be simpler and complete.

#Synced-CL 50276354

diff --git a/Engine/Plugins/Tests/EditorTests/Source/EditorTests/Private/UnrealEd/JsonUtilitiesTests.cpp b/Engine/Plugins/Tests/EditorTests/Source/EditorTests/Private/UnrealEd/JsonUtilitiesTests.cpp
--- a/Engine/Plugins/Tests/EditorTests/Source/EditorTests/Private/UnrealEd/JsonUtilitiesTests.cpp
+++ b/Engine/Plugins/Tests/EditorTests/Source/EditorTests/Private/UnrealEd/JsonUtilitiesTests.cpp
@@ -14,34 +14,6 @@
 
 // Expected values.
 
-namespace
-{
-	const FString ExpectedTestFloatDescription = FString(TEXT("TestFloatDesc"));
-	constexpr float ExpectedTestFloatMinimum = -101.0f;
-	constexpr float ExpectedTestFloatMaximum = 102.0f;
-
-	const FString ExpectedTestIntDescription = FString(TEXT("TestIntDesc"));
-	constexpr int32 ExpectedTestIntMinimum = -201;
-	constexpr int32 ExpectedTestIntMaximum = 202;
-
-	const FString ExpectedTestBoolDescription = FString(TEXT("TestBoolDesc"));
-
-	const FString ExpectedTestEnumDescription = FString(TEXT("TestEnumDesc"));
-	const TArray<FString> ExpectedTestEnumValues = {TEXT("Zero"), TEXT("One")};
-
-	const FString ExpectedTestStringDescription = FString(TEXT("TestStringDesc"));
-
-	const FString ExpectedTestNestedStructADescription = FString(TEXT("TestNestedStructADesc"));
-	const FString ExpectedTestNestedStructBDescription = FString(TEXT("TestNestedStructBDesc"));
-	const FString ExpectedTestNestedStructCDescription = FString(TEXT("TestNestedStructCDesc"));
-	const FString ExpectedTestNestedStructDDescription = FString(TEXT("TestNestedStructDDesc"));
-
-	const FString ExpectedTestFuncDescription = FString(TEXT("TestFuncDesc"));
-	const FString ExpectedTestFuncParamFloatDescription = FString(TEXT("TestFuncParamFloatDesc"));
-	const FString ExpectedTestFuncReturnValueDescription = FString(TEXT("TestFuncReturnValueDesc"));
-}
-
-
 // Tests.
 
 #if WITH_DEV_AUTOMATION_TESTS
@@ -52,137 +24,444 @@ PRAGMA_DISABLE_DEPRECATION_WARNINGS
 
 #define DEBUG_PRINT_JSONS_TO_LOG 0
 
-namespace
+
+BEGIN_DEFINE_SPEC(FJsonUtilitiesEditorSpec, "System.Promotion.Editor.JsonUtilities",
+	EAutomationTestFlags::EditorContext|EAutomationTestFlags::EngineFilter)
+END_DEFINE_SPEC(FJsonUtilitiesEditorSpec)
+
+void TestBasics(FAutomationTestBase& Test, const TSharedPtr<FJsonObject>& Schema, const FString& Type, const FString& Description)
 {
-	// Gets last member of the member path.
-	FString Last(const FJsonSchemaMemberPath& MemberPath)
+	Test.TestTrue("Valid schema", Schema.IsValid());
+	if (!Type.IsEmpty())
 	{
-		const FString MemberPathString(MemberPath.ToView());
-	
-		if (int32 Index = INDEX_NONE;
-			MemberPathString.FindLastChar(FJsonSchemaMemberPath::Delimiter[0], Index))
-		{
-			return MemberPathString.Mid(Index + 1);
-		}		
+		FString FieldType;
+		Test.TestTrue("Type", Schema->TryGetStringField("type", FieldType));
+		Test.TestEqual("Type", Type, FieldType);
+	}
+	if (!Description.IsEmpty())
+	{
+		FString FieldDescription;
+		Test.TestTrue("Description", Schema->TryGetStringField("description", FieldDescription));
+		Test.TestEqual("Description", Description, FieldDescription);
+	}
+}
+
+void TestFloat(FAutomationTestBase& Test, const TSharedPtr<FJsonObject>& Schema, const FString& Description,
+	TOptional<float> Min = TOptional<float>(), TOptional<float> Max = TOptional<float>(), TOptional<float> Default = TOptional<float>())
+{
+	TestBasics(Test, Schema, "number", Description);
+	if (Min)
+	{
+		float FieldMin;
+		Test.TestTrue("Min", Schema->TryGetNumberField("minimum", FieldMin));
+		Test.TestEqual("Min", Min.GetValue(), FieldMin);
+	}
+	if (Max)
+	{
+		float FieldMax;
+		Test.TestTrue("Max", Schema->TryGetNumberField("maximum", FieldMax));
+		Test.TestEqual("Max", Max.GetValue(), FieldMax);
+	}
+	if (Default)
+	{
+		float FieldDefault;
+		Test.TestTrue("Default", Schema->TryGetNumberField("default", FieldDefault));
+		Test.TestEqual("Default", Default.GetValue(), FieldDefault);
+	}
+}
 
-		return MemberPathString;
+void TestInteger(FAutomationTestBase& Test, const TSharedPtr<FJsonObject>& Schema, const FString& Description,
+	TOptional<int> Min = TOptional<int>(), TOptional<int> Max = TOptional<int>(), TOptional<int> Default = TOptional<int>())
+{
+	TestBasics(Test, Schema, "integer", Description);
+	if (Min)
+	{
+		int FieldMin;
+		Test.TestTrue("Min", Schema->TryGetNumberField("minimum", FieldMin));
+		Test.TestEqual("Min", Min.GetValue(), FieldMin);
 	}
-	
-	bool TestIsJsonSchemaValidAndNotEmpty(const TSharedPtr<FJsonObject>& Schema, FAutomationTestBase& Test)
+	if (Max)
 	{
-		return 
-			Test.TestTrue("Schema is valid.", Schema.IsValid()) &&
-			Test.TestTrue("Schema is not empty (has values.)", !Schema->Values.IsEmpty());
+		int FieldMax;
+		Test.TestTrue("Max", Schema->TryGetNumberField("maximum", FieldMax));
+		Test.TestEqual("Max", Max.GetValue(), FieldMax);
 	}
-	
-	bool TestIsJsonSchemaInvalid(const TSharedPtr<FJsonObject>& Schema, FAutomationTestBase& Test)
+	if (Default)
 	{
-		return Test.TestFalse("Schema is invalid.", Schema.IsValid());
+		int FieldDefault;
+		Test.TestTrue("Default", Schema->TryGetNumberField("default", FieldDefault));
+		Test.TestEqual("Default", Default.GetValue(), FieldDefault);
 	}
-	
-	TOptional<TSharedPtr<FJsonObject>> TestAndGetJsonObjectField(
-		const TSharedPtr<FJsonObject>& Schema, const FString& FieldName, FAutomationTestBase& Test)
+}
+
+void TestBoolean(FAutomationTestBase& Test, const TSharedPtr<FJsonObject>& Schema, const FString& Description, TOptional<bool> Default = TOptional<bool>())
+{
+	TestBasics(Test, Schema, "boolean", Description);
+	if (Default)
 	{
-		TOptional<TSharedPtr<FJsonObject>> OptionalJsonObject;
-		
-		if (Test.TestTrue("JSON schema is valid.", Schema.IsValid()))
+		bool FieldDefault;
+		Test.TestTrue("Default", Schema->TryGetBoolField("default", FieldDefault));
+		Test.TestEqual("Default", Default.GetValue(), FieldDefault);
+	}
+}
+
+void TestString(FAutomationTestBase& Test, const TSharedPtr<FJsonObject>& Schema, const FString& Description, TOptional<FString> Default= TOptional<FString>())
+{
+	TestBasics(Test, Schema, "string", Description);
+	if (Default)
+	{
+		FString FieldDefault;
+		Test.TestTrue("Default", Schema->TryGetStringField("default", FieldDefault));
+		Test.TestEqual("Default", Default.GetValue(), FieldDefault);
+	}
+}
+
+void TestEnum(FAutomationTestBase& Test, const TSharedPtr<FJsonObject>& Schema, const FString& Description, const TArray<FString>& Enums, TOptional<FString> Default = TOptional<FString>())
+{
+	TestBasics(Test, Schema, "string", Description);
+	if (!Enums.IsEmpty())
+	{
+		TArray<FString> FieldEnums;
+		Test.TestTrue("Enums", Schema->TryGetStringArrayField("enum", FieldEnums));
+		Test.TestEqual("Enums", Enums, FieldEnums);
+	}	
+	if (Default)
+	{
+		FString FieldDefault;
+		Test.TestTrue("Default", Schema->TryGetStringField("default", FieldDefault));
+		Test.TestEqual("Default", Default.GetValue(), FieldDefault);
+	}
+}
+
+void TestRequired(FAutomationTestBase& Test, const TSharedPtr<FJsonObject>& Schema, const TArray<FString>& Required)
+{
+	TArray<FString> RequiredField;	
+	Test.TestEqual("Required", Schema->TryGetStringArrayField("required", RequiredField), !Required.IsEmpty());
+	if (!Required.IsEmpty())
+	{
+		Test.TestEqual("Required", Required.Num(), RequiredField.Num());
+		for (auto& Item : Required)
+		{
+			Test.TestTrue("Required", RequiredField.Contains(Item));
+		}
+	}
+}
+
+void TestNumericProperties(FAutomationTestBase& Test, const TSharedPtr<FJsonObject>& Schema, const TArray<FString>& PropertyNames, bool FloatValues = true)
+{
+	const TSharedPtr<FJsonObject>* Properties = nullptr;
+	Test.TestTrue("Properties", Schema->TryGetObjectField("properties", Properties));
+	if (Properties)
+	{		
+		Test.TestEqual("Properties", Properties->Get()->Values.Num(), PropertyNames.Num());
+		for (auto& PropertyName : PropertyNames)
 		{
-			if (!Test.TestTrue(
-				FString::Printf(TEXT("%s object exists."), *FieldName),
-				Schema->HasTypedField(*FieldName, EJson::Object)))
+			const TSharedPtr<FJsonObject>* Property = nullptr;
+			Test.TestTrue("Properties", Properties->Get()->TryGetObjectField(PropertyName, Property));
+			if (Property)
 			{
-				return OptionalJsonObject;
+				if (FloatValues)
+				{
+					TestFloat(Test, *Property, "");
+				}
+				else
+				{
+					TestInteger(Test, *Property, "");
+				}				
 			}
-
-			OptionalJsonObject = Schema->GetObjectField(*FieldName);
 		}
-		
-		return OptionalJsonObject;
 	}
-	
-	TOptional<FString> TestAndGetJsonStringField(
-		const TSharedPtr<FJsonObject>& Schema, const FString& FieldName, FAutomationTestBase& Test)
+}
+
+void TestVector(FAutomationTestBase& Test, const TSharedPtr<FJsonObject>& Schema, const FString& Description, TOptional<FVector> Default = TOptional<FVector>())
+{
+	TestBasics(Test, Schema, "object", Description);
+	TArray<FString> PropertyNames({ "x", "y", "z" });
+	TestNumericProperties(Test, Schema, PropertyNames, true);
+	TestRequired(Test, Schema, PropertyNames);
+		
+	if(Default)
 	{
-		TOptional<FString> OptionalJsonString;
-		if (FString JsonString;
-			Test.TestTrue(
-				FString::Printf(TEXT("%s string exists."), *FieldName),
-				Schema->TryGetStringField(*FieldName, JsonString)))
+		const TSharedPtr<FJsonObject>* FieldDefault = nullptr;
+		Test.TestTrue("Default", Schema->TryGetObjectField("default", FieldDefault));
+		if (FieldDefault)
 		{
-			OptionalJsonString = JsonString;
+			Test.TestEqual("Default", FieldDefault->Get()->Values.Num(), PropertyNames.Num());
+			for (int i = 0; i < PropertyNames.Num(); ++i)
+			{
+				float PropertyValue;
+				Test.TestTrue("Default", FieldDefault->Get()->TryGetNumberField(PropertyNames[i], PropertyValue));
+				Test.TestEqual("Default", PropertyValue, (float)Default.GetValue()[i]);
+			}
 		}
-		
-		return OptionalJsonString;
 	}
-	
-	TOptional<TArray<FString>> TestAndGetJsonStringArrayField(
-		const TSharedPtr<FJsonObject>& Schema, const FString& FieldName, FAutomationTestBase& Test)
+}
+
+void TestQuat(FAutomationTestBase& Test, const TSharedPtr<FJsonObject>& Schema, const FString& Description, TOptional<FQuat> Default = TOptional<FQuat>())
+{
+	TestBasics(Test, Schema, "object", Description);
+	TArray<FString> PropertyNames({ "x", "y", "z", "w"});
+	TestNumericProperties(Test, Schema, PropertyNames, true);
+	TestRequired(Test, Schema, PropertyNames);
+
+	if (Default)
 	{
-		TOptional<TArray<FString>> OptionalJsonStringArray;
-		if (TArray<FString> JsonStringArray;
-			Test.TestTrue(
-				FString::Printf(TEXT("%s string array exists."), *FieldName),
-				Schema->TryGetStringArrayField(*FieldName, JsonStringArray)))
+		const TSharedPtr<FJsonObject>* FieldDefault = nullptr;
+		Test.TestTrue("Default", Schema->TryGetObjectField("default", FieldDefault));
+		if (FieldDefault)
 		{
-			OptionalJsonStringArray = JsonStringArray;
+			Test.TestEqual("Default", FieldDefault->Get()->Values.Num(), PropertyNames.Num());
+			for (auto& PropertyName : PropertyNames)
+			{
+				double PropertyValue;
+				Test.TestTrue("Default", FieldDefault->Get()->TryGetNumberField(PropertyName, PropertyValue));
+				if (PropertyName == "x")
+				{
+					Test.TestEqual("Default", PropertyValue, Default.GetValue().X);
+				}
+				else if (PropertyName == "y")
+				{
+					Test.TestEqual("Default", PropertyValue, Default.GetValue().Y);
+				}
+				else if (PropertyName == "z")
+				{
+					Test.TestEqual("Default", PropertyValue, Default.GetValue().Z);
+				}
+				else if (PropertyName == "w")
+				{
+					Test.TestEqual("Default", PropertyValue, Default.GetValue().W);
+				}
+			}
 		}
-		
-		return OptionalJsonStringArray;
 	}
-	
-	TOptional<float> TestAndGetJsonFloatNumberField(
-		const TSharedPtr<FJsonObject>& Schema, const FString& FieldName, FAutomationTestBase& Test)
+}
+
+void TestRotator(FAutomationTestBase& Test, const TSharedPtr<FJsonObject>& Schema, const FString& Description, TOptional<FRotator> Default = TOptional<FRotator>())
+{
+	TestBasics(Test, Schema, "object", Description);
+	TArray<FString> PropertyNames({ "yaw", "pitch", "roll" });
+	TestNumericProperties(Test, Schema, PropertyNames, true);
+	TestRequired(Test, Schema, PropertyNames);
+
+	if (Default)
 	{
-		TOptional<float> OptionalJsonNumber;
-		if (float JsonNumber = 0.0f;
-			Test.TestTrue(
-				FString::Printf(TEXT("%s float number exists."), *FieldName),
-				Schema->TryGetNumberField(*FieldName, JsonNumber)))
+		const TSharedPtr<FJsonObject>* FieldDefault = nullptr;
+		Test.TestTrue("Default", Schema->TryGetObjectField("default", FieldDefault));
+		if (FieldDefault)
 		{
-			OptionalJsonNumber = JsonNumber;
+			Test.TestEqual("Default", FieldDefault->Get()->Values.Num(), PropertyNames.Num());
+			for (auto& PropertyName : PropertyNames)
+			{
+				double PropertyValue;
+				Test.TestTrue("Default", FieldDefault->Get()->TryGetNumberField(PropertyName, PropertyValue));				
+				if (PropertyName == "yaw")
+				{
+					Test.TestEqual("Default", PropertyValue, Default.GetValue().Yaw);
+				}
+				else if (PropertyName == "pitch")
+				{
+					Test.TestEqual("Default", PropertyValue, Default.GetValue().Pitch);
+				}
+				else if (PropertyName == "roll")
+				{
+					Test.TestEqual("Default", PropertyValue, Default.GetValue().Roll);
+				}
+			}
 		}
-		
-		return OptionalJsonNumber;
 	}
-	
-	TOptional<int32> TestAndGetJsonIntNumberField(
-		const TSharedPtr<FJsonObject>& Schema, const FString& FieldName, FAutomationTestBase& Test)
+}
+
+void TestTransform(FAutomationTestBase& Test, const TSharedPtr<FJsonObject>& Schema, const FString& Description, TOptional<FTransform> Default = TOptional<FTransform>())
+{
+	TestBasics(Test, Schema, "object", Description);
+	TArray<FString> PropertyNames({ "translation", "rotation", "scale3D" });
+	TestRequired(Test, Schema, PropertyNames);
+
+	const TSharedPtr<FJsonObject>* Properties = nullptr;
+	Test.TestTrue("Properties", Schema->TryGetObjectField("properties", Properties));
+	if (Properties)
 	{
-		TOptional<int32> OptionalJsonNumber;
-		if (int32 JsonNumber = 0;
-			Test.TestTrue(
-				FString::Printf(TEXT("%s int number exists."), *FieldName),
-				Schema->TryGetNumberField(*FieldName, JsonNumber)))
+		Test.TestEqual("Properties", Properties->Get()->Values.Num(), PropertyNames.Num());
+		const TSharedPtr<FJsonObject>* Translation = nullptr;
+		Test.TestTrue("Input", Properties->Get()->TryGetObjectField("translation", Translation));
+		if (Translation)
 		{
-			OptionalJsonNumber = JsonNumber;
+			TestVector(Test, *Translation, "");
+		}
+
+		const TSharedPtr<FJsonObject>* Quat = nullptr;
+		Test.TestTrue("Input", Properties->Get()->TryGetObjectField("rotation", Quat));
+		if (Quat)
+		{
+			TestQuat(Test, *Quat, "");
+		}
+
+		const TSharedPtr<FJsonObject>* Scale = nullptr;
+		Test.TestTrue("Input", Properties->Get()->TryGetObjectField("scale3D", Scale));
+		if (Scale)
+		{
+			TestVector(Test, *Scale, "");
 		}
-		
-		return OptionalJsonNumber;
 	}
 	
-	void PrintJsonToLog(
-		const TCHAR* Name, const TSharedPtr<FJsonObject>& JsonObject)
+	if (Default)
 	{
-		if (JsonObject)
+		const TSharedPtr<FJsonObject>* FieldDefault = nullptr;
+		Test.TestTrue("Default", Schema->TryGetObjectField("default", FieldDefault));
+		if (FieldDefault)
 		{
-			FString JsonString;
-			const TSharedRef<TJsonWriter<>> JsonWriter = TJsonWriterFactory<>::Create(&JsonString);
-			FJsonSerializer::Serialize(JsonObject, JsonWriter);
-			UE_LOG(LogJson, Display, TEXT("Json \"%s\" is valid -\n%s"), Name, *JsonString);
+			Test.TestEqual("Properties", FieldDefault->Get()->Values.Num(), PropertyNames.Num());
+			{
+				const TSharedPtr<FJsonObject>* Translation = nullptr;
+				Test.TestTrue("Default", FieldDefault->Get()->TryGetObjectField("translation", Translation));
+				if (Translation)
+				{
+					UScriptStruct* VectorStruct = TBaseStructure<FVector>::Get();
+					FVector ParsedVector;
+					FJsonObjectConverter::JsonObjectToUStruct(
+						Translation->ToSharedRef(),
+						VectorStruct,
+						&ParsedVector
+					);
+					Test.TestEqual("Default", ParsedVector, Default.GetValue().GetLocation());
+				}
+			}
+			
+			{
+				const TSharedPtr<FJsonObject>* Quat = nullptr;
+				Test.TestTrue("Default", FieldDefault->Get()->TryGetObjectField("rotation", Quat));
+				if (Quat)
+				{
+					UScriptStruct* QuatStruct = TBaseStructure<FQuat>::Get();
+					FQuat ParsedQuat;
+					FJsonObjectConverter::JsonObjectToUStruct(
+						Quat->ToSharedRef(),
+						QuatStruct,
+						&ParsedQuat
+					);
+					Test.TestEqual("Default", ParsedQuat, Default.GetValue().GetRotation());
+				}
+			}
+			
+			{
+				const TSharedPtr<FJsonObject>* Scale = nullptr;
+				Test.TestTrue("Default", FieldDefault->Get()->TryGetObjectField("scale3D", Scale));
+				if (Scale)
+				{
+					UScriptStruct* VectorStruct = TBaseStructure<FVector>::Get();
+					FVector ParsedVector;
+					FJsonObjectConverter::JsonObjectToUStruct(
+						Scale->ToSharedRef(),
+						VectorStruct,
+						&ParsedVector
+					);
+					Test.TestEqual("Default", ParsedVector, Default.GetValue().GetScale3D());
+				}
+			}			
 		}
-		else
+	}	
+}
+
+void TestLinearColor(FAutomationTestBase& Test, const TSharedPtr<FJsonObject>& Schema, const FString& Description, TOptional<FLinearColor> Default = TOptional<FLinearColor>())
+{
+	TestBasics(Test, Schema, "object", Description);
+	TArray<FString> PropertyNames({ "r", "g", "b", "a" });
+	TestNumericProperties(Test, Schema, PropertyNames, true);
+	TestRequired(Test, Schema, PropertyNames);
+
+	if (Default)
+	{
+		const TSharedPtr<FJsonObject>* FieldDefault = nullptr;
+		Test.TestTrue("Default", Schema->TryGetObjectField("default", FieldDefault));
+		if (FieldDefault)
 		{
-			UE_LOG(LogJson, Error, TEXT("Json \"%s\" is invalid."), Name);
+			Test.TestEqual("Default", FieldDefault->Get()->Values.Num(), PropertyNames.Num());
+			for (auto& PropertyName : PropertyNames)
+			{
+				float PropertyValue;
+				Test.TestTrue("Default", FieldDefault->Get()->TryGetNumberField(PropertyName, PropertyValue));
+				if (PropertyName == "r")
+				{
+					Test.TestEqual("Default", PropertyValue, Default.GetValue().R);
+				}
+				else if (PropertyName == "g")
+				{
+					Test.TestEqual("Default", PropertyValue, Default.GetValue().G);
+				}
+				else if (PropertyName == "b")
+				{
+					Test.TestEqual("Default", PropertyValue, Default.GetValue().B);
+				}
+				else if (PropertyName == "a")
+				{
+					Test.TestEqual("Default", PropertyValue, Default.GetValue().A);
+				}
+			}
 		}
 	}
 }
 
-BEGIN_DEFINE_SPEC(FJsonUtilitiesEditorSpec, "System.Promotion.Editor.JsonUtilities",
-	EAutomationTestFlags::EditorContext|EAutomationTestFlags::EngineFilter)
-END_DEFINE_SPEC(FJsonUtilitiesEditorSpec)
+void TestColor(FAutomationTestBase& Test, const TSharedPtr<FJsonObject>& Schema, const FString& Description, TOptional<FColor> Default = TOptional<FColor>())
+{
+	TestBasics(Test, Schema, "object", Description);
+	TArray<FString> PropertyNames({ "r", "g", "b", "a" });
+	TestNumericProperties(Test, Schema, PropertyNames, false);
+	TestRequired(Test, Schema, PropertyNames);
 
-void FJsonUtilitiesEditorSpec::Define()
+	if (Default)
+	{
+		const TSharedPtr<FJsonObject>* FieldDefault = nullptr;
+		Test.TestTrue("Default", Schema->TryGetObjectField("default", FieldDefault));
+		if (FieldDefault)
+		{
+			Test.TestEqual("Default", FieldDefault->Get()->Values.Num(), PropertyNames.Num());
+			for (auto& PropertyName : PropertyNames)
+			{
+				uint8 PropertyValue;
+				Test.TestTrue("Default", FieldDefault->Get()->TryGetNumberField(PropertyName, PropertyValue));
+				if (PropertyName == "r")
+				{
+					Test.TestEqual("Default", PropertyValue, Default.GetValue().R);
+				}
+				else if (PropertyName == "g")
+				{
+					Test.TestEqual("Default", PropertyValue, Default.GetValue().G);
+				}
+				else if (PropertyName == "b")
+				{
+					Test.TestEqual("Default", PropertyValue, Default.GetValue().B);
+				}
+				else if (PropertyName == "a")
+				{
+					Test.TestEqual("Default", PropertyValue, Default.GetValue().A);
+				}
+			}
+		}
+	}
+}
+
+bool TestCustomCallback(const FProperty* Property, const FString& ParameterDefaultString, const TSharedRef<FJsonObject>& OutputSchema)
 {
+	if (const FStructProperty* StructProperty = CastField<FStructProperty>(Property))
+	{	
+		if (StructProperty->Struct == TBaseStructure<FSoftClassPath>::Get())
+		{
+			if (ParameterDefaultString.IsEmpty())
+			{
+				OutputSchema->SetStringField(TEXT("type"), TEXT("string"));
+			}
+			else
+			{
+				OutputSchema->SetStringField(TEXT("default"), ParameterDefaultString);
+			}
+			return true;
+		}		
+	}
+	return false;
+}
+
+void FJsonUtilitiesEditorSpec::Define()
+{	
 	BeforeEach([this]() -> void
 		{
 			// Loading this module will enable the collection of metadata for the schemas.
@@ -198,644 +477,468 @@ void FJsonUtilitiesEditorSpec::Define()
 		[this]() -> void
 		{
 			It(
-				"Should generate valid JSON schema from a UStruct - smoke test.", 
+				"Should generate valid JSON schema from float property.",
 				[this]() -> void
 				{
-					const TSharedPtr<FJsonObject> Schema = FJsonSchemaGenerator::UStructToJsonSchemaObject<FJsonUtilitiesFakeStruct>();
-					if (!TestIsJsonSchemaValidAndNotEmpty(Schema, *this))
-					{
-						return;
-					}
+					UStruct* Struct = FJsonUtilitiesFakeStruct::StaticStruct();
+					FProperty* Prop = Struct->FindPropertyByName("TestFloat");
+					const TSharedPtr<FJsonObject> Schema = FJsonSchemaGenerator::FPropertyToJsonSchemaObject(Prop);
+					TestFloat(*this, Schema, "TestFloatDesc", -101.f, 102.f);
 				});
-			
+
 			It(
-				"Should generate valid JSON schema from a UClass - smoke test.", 
+				"Should generate valid JSON schema from int property.",
 				[this]() -> void
 				{
-					const TSharedPtr<FJsonObject> Schema = FJsonSchemaGenerator::UStructToJsonSchemaObject<UJsonUtilitiesFakeClass>();
-					if (!TestIsJsonSchemaValidAndNotEmpty(Schema, *this))
-					{
-						return;
-					}
+					UStruct* Struct = FJsonUtilitiesFakeStruct::StaticStruct();
+					FProperty* Prop = Struct->FindPropertyByName("TestInt");
+					const TSharedPtr<FJsonObject> Schema = FJsonSchemaGenerator::FPropertyToJsonSchemaObject(Prop);
+					TestInteger(*this, Schema, "TestIntDesc", -201, 202);
 				});
-			
+
+			It(
+				"Should generate valid JSON schema from bool property.",
+				[this]() -> void
+				{
+					UStruct* Struct = FJsonUtilitiesFakeStruct::StaticStruct();
+					FProperty* Prop = Struct->FindPropertyByName("TestBool");
+					const TSharedPtr<FJsonObject> Schema = FJsonSchemaGenerator::FPropertyToJsonSchemaObject(Prop);
+					TestBoolean(*this, Schema, "TestBoolDesc");
+				});
+
+			It(
+				"Should generate valid JSON schema from string property.",
+				[this]() -> void
+				{
+					UStruct* Struct = FJsonUtilitiesFakeStruct::StaticStruct();
+					FProperty* Prop = Struct->FindPropertyByName("TestString");
+					const TSharedPtr<FJsonObject> Schema = FJsonSchemaGenerator::FPropertyToJsonSchemaObject(Prop);
+					TestString(*this, Schema, "TestStringDesc");
+				});
+
+			It(
+				"Should generate valid JSON schema from name property.",
+				[this]() -> void
+				{
+					UStruct* Struct = FJsonUtilitiesFakeStruct::StaticStruct();
+					FProperty* Prop = Struct->FindPropertyByName("TestName");
+					const TSharedPtr<FJsonObject> Schema = FJsonSchemaGenerator::FPropertyToJsonSchemaObject(Prop);
+					TestString(*this, Schema, "TestNameDesc");
+				});
+
+			It(
+				"Should generate valid JSON schema from enum property.",
+				[this]() -> void
+				{
+					UStruct* Struct = FJsonUtilitiesFakeStruct::StaticStruct();
+					FProperty* Prop = Struct->FindPropertyByName("TestEnum");
+					const TSharedPtr<FJsonObject> Schema = FJsonSchemaGenerator::FPropertyToJsonSchemaObject(Prop);
+					TestEnum(*this, Schema, "TestEnumDesc", TArray<FString>({ "Zero", "One" }));
+				});
+
+			It(
+				"Should generate valid JSON schema from UObject property.",
+				[this]() -> void
+				{
+					UStruct* Struct = FJsonUtilitiesFakeStruct::StaticStruct();
+					FProperty* Prop = Struct->FindPropertyByName("TestObject");
+					const TSharedPtr<FJsonObject> Schema = FJsonSchemaGenerator::FPropertyToJsonSchemaObject(Prop);
+					TestString(*this, Schema, "TestObjectDesc");
+				});
+
+			It(
+				"Should generate valid JSON schema from UClass property.",
+				[this]() -> void
+				{
+					UStruct* Struct = FJsonUtilitiesFakeStruct::StaticStruct();
+					FProperty* Prop = Struct->FindPropertyByName("TestClass");
+					const TSharedPtr<FJsonObject> Schema = FJsonSchemaGenerator::FPropertyToJsonSchemaObject(Prop);
+					TestString(*this, Schema, "TestClassDesc");
+				});
+
 			It(
-				"Should generate valid JSON schema from a UStruct's float UProperty.", 
+				"Should generate valid JSON schema from Struct property.",
 				[this]() -> void
 				{
-					FJsonSchemaMemberPath MemberPath(TEXT("[schema]")); // ..mostly for error reporting
-					
-					const TSharedPtr<FJsonObject> Schema = FJsonSchemaGenerator::UStructToJsonSchemaObject<FJsonUtilitiesFakeStruct>();
-					if (!TestIsJsonSchemaValidAndNotEmpty(Schema, *this))
+					UStruct* Struct = FJsonUtilitiesFakeStruct::StaticStruct();
+					FProperty* Prop = Struct->FindPropertyByName("TestNestedStruct");
+					const TSharedPtr<FJsonObject> Schema = FJsonSchemaGenerator::FPropertyToJsonSchemaObject(Prop);
+					TestBasics(*this, Schema, "object", "TestNestedStructDesc");
+					const TSharedPtr<FJsonObject>* Properties = nullptr;
+					TestTrue("Properties", Schema->TryGetObjectField("properties", Properties));
+					if (Properties)
 					{
-						return;
+						const TSharedPtr<FJsonObject>* Property = nullptr;
+						TestTrue("Properties", Properties->Get()->TryGetObjectField("TestNestedStructFloat", Property));
+						if (Property)
+						{
+							TestFloat(*this, *Property, "TestNestedFloatDesc", -301.f, 302.f);
+							Property = nullptr;
+						}
 					}
-						
-					UE_JSON_SCHEMA_SCOPED_MEMBER_PATH_PUSH(MemberPath, TEXT("properties"), true);
-					
-					if (const TOptional<TSharedPtr<FJsonObject>> Schema_Properties = TestAndGetJsonObjectField(
-							Schema, Last(MemberPath), *this);
-						Schema_Properties.IsSet())
+				});
+
+			It(
+				"Should generate valid JSON schema from structs.",
+				[this]() -> void
+				{
+					UStruct* Struct = FJsonUtilitiesFakeStruct::StaticStruct();
+					const TSharedPtr<FJsonObject> Schema = FJsonSchemaGenerator::UStructToJsonSchemaObject(Struct);
+					TestBasics(*this, Schema, "object", "TestStructDesc");
+					const TSharedPtr<FJsonObject>* Properties = nullptr;
+					TestTrue("Properties", Schema->TryGetObjectField("properties", Properties));
+					if (Properties)
 					{
-						UE_JSON_SCHEMA_SCOPED_MEMBER_PATH_PUSH(MemberPath, TEXT("testFloat"), true);
-						
-						if (const TOptional<TSharedPtr<FJsonObject>> Schema_Properties_TestFloat = TestAndGetJsonObjectField(
-								*Schema_Properties, Last(MemberPath), *this);
-							Schema_Properties_TestFloat.IsSet())
 						{
+							const TSharedPtr<FJsonObject>* Property = nullptr;
+							TestTrue("Properties", Properties->Get()->TryGetObjectField("TestFloat", Property));
+							if (Property)
 							{
-								UE_JSON_SCHEMA_SCOPED_MEMBER_PATH_PUSH(MemberPath, TEXT("type"), true);
-								
-								if (const TOptional<FString> Schema_Properties_TestFloat_Type = TestAndGetJsonStringField(
-										*Schema_Properties_TestFloat, Last(MemberPath), *this);
-									Schema_Properties_TestFloat_Type.IsSet())
-								{
-									TestEqual(
-										FString::Printf(TEXT("%s is number."), *MemberPath), 
-											(*Schema_Properties_TestFloat_Type), 
-										"number");
-								}
+								TestFloat(*this, *Property, "TestFloatDesc", -101.f, 102.f);
 							}
+						}
 
+						{
+							const TSharedPtr<FJsonObject>* Property = nullptr;
+							TestTrue("Properties", Properties->Get()->TryGetObjectField("TestInt", Property));
+							if (Property)
 							{
-								UE_JSON_SCHEMA_SCOPED_MEMBER_PATH_PUSH(MemberPath, TEXT("minimum"), true);
-								
-								if (const TOptional<float> Schema_Properties_TestFloat_Minimum = TestAndGetJsonFloatNumberField(
-										*Schema_Properties_TestFloat, Last(MemberPath), *this);
-									Schema_Properties_TestFloat_Minimum.IsSet())
-								{
-									TestEqual(
-										FString::Printf(TEXT("%s should be %f. Is %f."), *MemberPath, 
-											ExpectedTestFloatMinimum, *Schema_Properties_TestFloat_Minimum),
-										*Schema_Properties_TestFloat_Minimum, 
-										ExpectedTestFloatMinimum); 
-								}
+								TestInteger(*this, *Property, "TestIntDesc", -201, 202);
 							}
-							
+						}
+
+						{
+							const TSharedPtr<FJsonObject>* Property = nullptr;
+							TestTrue("Properties", Properties->Get()->TryGetObjectField("TestBool", Property));
+							if (Property)
 							{
-								UE_JSON_SCHEMA_SCOPED_MEMBER_PATH_PUSH(MemberPath, TEXT("maximum"), true);
-								
-								if (const TOptional<float> Schema_Properties_TestFloat_Maximum = TestAndGetJsonFloatNumberField(
-										*Schema_Properties_TestFloat, Last(MemberPath), *this);
-									Schema_Properties_TestFloat_Maximum.IsSet())
-								{
-									TestEqual(
-										FString::Printf(TEXT("%s should be %f. Is %f."), *MemberPath, 
-											ExpectedTestFloatMaximum,* Schema_Properties_TestFloat_Maximum),
-										*Schema_Properties_TestFloat_Maximum, 
-										ExpectedTestFloatMaximum); 
-								}
+								TestBoolean(*this, *Property, "TestBoolDesc");
 							}
-							
+						}
+
+						{
+							const TSharedPtr<FJsonObject>* Property = nullptr;
+							TestTrue("Properties", Properties->Get()->TryGetObjectField("TestString", Property));
+							if (Property)
 							{
-								UE_JSON_SCHEMA_SCOPED_MEMBER_PATH_PUSH(MemberPath, TEXT("description"), true);
-								
-								if (const TOptional<FString> Schema_Properties_TestFloat_Description = TestAndGetJsonStringField(
-										*Schema_Properties_TestFloat, Last(MemberPath), *this);
-									Schema_Properties_TestFloat_Description.IsSet())
-								{
-									TestEqual(
-										FString::Printf(TEXT("%s should be '%s'. Is '%s'."), *MemberPath, 
-											*ExpectedTestFloatDescription, **Schema_Properties_TestFloat_Description),
-										*Schema_Properties_TestFloat_Description, 
-										ExpectedTestFloatDescription); 
-								}
+								TestString(*this, *Property, "TestStringDesc");
+							}
+						}
+
+						{
+							const TSharedPtr<FJsonObject>* Property = nullptr;
+							TestTrue("Properties", Properties->Get()->TryGetObjectField("TestName", Property));
+							if (Property)
+							{
+								TestString(*this, *Property, "TestNameDesc");
+							}
+						}
+
+						{
+							const TSharedPtr<FJsonObject>* Property = nullptr;
+							TestTrue("Properties", Properties->Get()->TryGetObjectField("TestObject", Property));
+							if (Property)
+							{
+								TestString(*this, *Property, "TestObjectDesc");
+							}
+						}
+
+						{
+							const TSharedPtr<FJsonObject>* Property = nullptr;
+							TestTrue("Properties", Properties->Get()->TryGetObjectField("TestClass", Property));
+							if (Property)
+							{
+								TestString(*this, *Property, "TestClassDesc");
+							}
+						}
+
+						{
+							const TSharedPtr<FJsonObject>* Property = nullptr;
+							TestTrue("Properties", Properties->Get()->TryGetObjectField("TestNestedStruct", Property));
+							if (Property)
+							{
+								TestBasics(*this, *Property, "object", "TestNestedStructDesc");
 							}
 						}
 					}
 				});
-			
+
 			It(
-				"Should generate valid JSON schema from a UStruct's int UProperty.", 
+				"Should generate valid JSON schema from a UClass.",
 				[this]() -> void
 				{
-					FJsonSchemaMemberPath MemberPath(TEXT("[schema]")); // ..mostly for error reporting
-					
-					const TSharedPtr<FJsonObject> Schema = FJsonSchemaGenerator::UStructToJsonSchemaObject<FJsonUtilitiesFakeStruct>();
-					if (!TestIsJsonSchemaValidAndNotEmpty(Schema, *this))
+					const TSharedPtr<FJsonObject> Schema = FJsonSchemaGenerator::UStructToJsonSchemaObject<UJsonUtilitiesFakeClass>();
+					TestBasics(*this, Schema, "object", "TestClassDesc");
+					const TSharedPtr<FJsonObject>* Properties = nullptr;
+					TestTrue("Properties", Schema->TryGetObjectField("properties", Properties));
+					if (Properties)
 					{
-						return;
+						const TSharedPtr<FJsonObject>* Property = nullptr;
+						TestTrue("Properties", Properties->Get()->TryGetObjectField("TestFloat", Property));
+						if (Property)
+						{
+							TestFloat(*this, *Property, "TestFloatDesc", -101.f, 102.f);
+						}
 					}
-						
-					UE_JSON_SCHEMA_SCOPED_MEMBER_PATH_PUSH(MemberPath, TEXT("properties"), true);
-					
-					if (const TOptional<TSharedPtr<FJsonObject>> Schema_Properties = TestAndGetJsonObjectField(
-							Schema, Last(MemberPath), *this);
-						Schema_Properties.IsSet())
+				});
+
+			It(
+				"Should generate JSON schema from a basic UFunction.",
+				[this]() -> void
+				{
+					UFunction* FunctionPtr = UJsonUtilitiesFakeClass::StaticClass()->FindFunctionByName("TestFunc");
+					const TSharedPtr<FJsonObject> Schema = FJsonSchemaGenerator::UStructToJsonSchemaObject(FunctionPtr);
+					TestTrue("Schema", Schema.IsValid());
+					TestBasics(*this, Schema, "", "TestFuncDesc");
+					const TSharedPtr<FJsonObject>* InputSchema = nullptr;
+					TestTrue("Input", Schema->TryGetObjectField("inputSchema", InputSchema));
+					if (InputSchema)
 					{
-						UE_JSON_SCHEMA_SCOPED_MEMBER_PATH_PUSH(MemberPath, TEXT("testInt"), true);
-						
-						if (const TOptional<TSharedPtr<FJsonObject>> Schema_Properties_TestInt = TestAndGetJsonObjectField(
-								*Schema_Properties, Last(MemberPath), *this);
-							Schema_Properties_TestInt.IsSet())
+						TestBasics(*this, *InputSchema, "object", "");
+
+						const TSharedPtr<FJsonObject>* Properties = nullptr;
+						TestTrue("Input", InputSchema->Get()->TryGetObjectField("properties", Properties));
+						if (Properties)
 						{
+							TestEqual("Input", Properties->Get()->Values.Num(), 12);
+							TestRequired(*this, *InputSchema, TArray<FString>({ "requiredParam" }));
+
 							{
-								UE_JSON_SCHEMA_SCOPED_MEMBER_PATH_PUSH(MemberPath, TEXT("type"), true);
-								
-								if (const TOptional<FString> Schema_Properties_TestInt_Type = TestAndGetJsonStringField(
-										*Schema_Properties_TestInt, Last(MemberPath), *this);
-									Schema_Properties_TestInt_Type.IsSet())
+								const TSharedPtr<FJsonObject>* Property = nullptr;
+								TestTrue("Input", Properties->Get()->TryGetObjectField("requiredParam", Property));
+								if (Property)
 								{
-									TestTrue(
-										FString::Printf(TEXT("%s is integer."), *MemberPath),
-										*Schema_Properties_TestInt_Type == "integer");
+									TestInteger(*this, *Property, "TestRequiredDesc");
 								}
 							}
 
 							{
-								UE_JSON_SCHEMA_SCOPED_MEMBER_PATH_PUSH(MemberPath, TEXT("minimum"), true);
-								
-								if (const TOptional<int32> Schema_Properties_TestInt_Minimum = TestAndGetJsonIntNumberField(
-										*Schema_Properties_TestInt, Last(MemberPath), *this);
-									Schema_Properties_TestInt_Minimum.IsSet())
+								const TSharedPtr<FJsonObject>* Property = nullptr;
+								TestTrue("Input", Properties->Get()->TryGetObjectField("boolParam", Property));
+								if (Property)
 								{
-									TestEqual(
-										FString::Printf(TEXT("%s should be %d. Is %d."), *MemberPath, 
-											ExpectedTestIntMinimum, *Schema_Properties_TestInt_Minimum),
-										*Schema_Properties_TestInt_Minimum, 
-										ExpectedTestIntMinimum); 
+									TestBoolean(*this, *Property, "TestBoolDesc", true);
 								}
 							}
-							
+
 							{
-								UE_JSON_SCHEMA_SCOPED_MEMBER_PATH_PUSH(MemberPath, TEXT("maximum"), true);
-								
-								if (const TOptional<int32> Schema_Properties_TestInt_Maximum = TestAndGetJsonIntNumberField(
-										*Schema_Properties_TestInt, Last(MemberPath), *this);
-									Schema_Properties_TestInt_Maximum.IsSet())
+								const TSharedPtr<FJsonObject>* Property = nullptr;
+								TestTrue("Input", Properties->Get()->TryGetObjectField("floatParam", Property));
+								if (Property)
 								{
-									TestEqual(
-										FString::Printf(TEXT("%s should be %d. Is %d."), *MemberPath, 
-											ExpectedTestIntMaximum, *Schema_Properties_TestInt_Maximum),
-										*Schema_Properties_TestInt_Maximum, 
-										ExpectedTestIntMaximum); 
+									TestFloat(*this, *Property, "TestFloatDesc", TOptional<float>(), TOptional<float>(), 1.23f);
 								}
 							}
-							
+
 							{
-								UE_JSON_SCHEMA_SCOPED_MEMBER_PATH_PUSH(MemberPath, TEXT("description"), true);
-								
-								if (const TOptional<FString> Schema_Properties_TestInt_Description = TestAndGetJsonStringField(
-										*Schema_Properties_TestInt, Last(MemberPath), *this);
-									Schema_Properties_TestInt_Description.IsSet())
+								const TSharedPtr<FJsonObject>* Property = nullptr;
+								TestTrue("Input", Properties->Get()->TryGetObjectField("enumParam", Property));
+								if (Property)
 								{
-									TestEqual(
-										FString::Printf(TEXT("%s should be '%s'. Is '%s'."), *MemberPath, 
-											*ExpectedTestIntDescription, **Schema_Properties_TestInt_Description),
-										*Schema_Properties_TestInt_Description, 
-										ExpectedTestIntDescription); 
+									TestEnum(*this, *Property, "TestEnumDesc", TArray<FString>(), FString("One"));
 								}
 							}
-						}
-					}
-				});
-			
-			It(
-				"Should generate valid JSON schema from a UStruct's bool UProperty.", 
-				[this]() -> void
-				{
-					FJsonSchemaMemberPath MemberPath(TEXT("[schema]")); // ..mostly for error reporting
-					
-					const TSharedPtr<FJsonObject> Schema = FJsonSchemaGenerator::UStructToJsonSchemaObject<FJsonUtilitiesFakeStruct>();
-					if (!TestIsJsonSchemaValidAndNotEmpty(Schema, *this))
-					{
-						return;
-					}
 
-					UE_JSON_SCHEMA_SCOPED_MEMBER_PATH_PUSH(MemberPath, TEXT("properties"), true);
-					
-					if (const TOptional<TSharedPtr<FJsonObject>> Schema_Properties = TestAndGetJsonObjectField(
-							Schema, Last(MemberPath), *this);
-						Schema_Properties.IsSet())
-					{
-						UE_JSON_SCHEMA_SCOPED_MEMBER_PATH_PUSH(MemberPath, TEXT("testBool"), true);
-						
-						if (const TOptional<TSharedPtr<FJsonObject>> Schema_Properties_TestBool = TestAndGetJsonObjectField(
-								*Schema_Properties, Last(MemberPath), *this);
-							Schema_Properties_TestBool.IsSet())
-						{
 							{
-								UE_JSON_SCHEMA_SCOPED_MEMBER_PATH_PUSH(MemberPath, TEXT("type"), true);
-								
-								if (const TOptional<FString> Schema_Properties_TestBool_Type = TestAndGetJsonStringField(
-										*Schema_Properties_TestBool, Last(MemberPath), *this);
-									Schema_Properties_TestBool_Type.IsSet())
+								const TSharedPtr<FJsonObject>* Property = nullptr;
+								TestTrue("Input", Properties->Get()->TryGetObjectField("stringParam", Property));
+								if (Property)
 								{
-									TestTrue(
-										FString::Printf(TEXT("%s is boolean."), *MemberPath),
-										*Schema_Properties_TestBool_Type == "boolean"); 
+									TestString(*this, *Property, "TestStringDesc", FString("foo"));
 								}
 							}
-							
+
 							{
-								UE_JSON_SCHEMA_SCOPED_MEMBER_PATH_PUSH(MemberPath, TEXT("description"), true);
-								
-								if (const TOptional<FString> Schema_Properties_TestBool_Description = TestAndGetJsonStringField(
-										*Schema_Properties_TestBool, Last(MemberPath), *this);
-									Schema_Properties_TestBool_Description.IsSet())
+								const TSharedPtr<FJsonObject>* Property = nullptr;
+								TestTrue("Input", Properties->Get()->TryGetObjectField("nameParam", Property));
+								if (Property)
 								{
-									TestEqual(
-										FString::Printf(TEXT("%s should be '%s'. Is '%s'."), *MemberPath, 
-											*ExpectedTestBoolDescription, **Schema_Properties_TestBool_Description),
-										*Schema_Properties_TestBool_Description, 
-										ExpectedTestBoolDescription); 
+									TestString(*this, *Property, "TestNameDesc", FString("bar"));
 								}
 							}
-						}
-					}
-				});
-			
-			It(
-				"Should generate valid JSON schema from a UStruct's string UProperty.", 
-				[this]() -> void
-				{
-					FJsonSchemaMemberPath MemberPath(TEXT("[schema]")); // ..mostly for error reporting
-					
-					const TSharedPtr<FJsonObject> Schema = FJsonSchemaGenerator::UStructToJsonSchemaObject<FJsonUtilitiesFakeStruct>();
-					if (!TestIsJsonSchemaValidAndNotEmpty(Schema, *this))
-					{
-						return;
-					}
-						
-					UE_JSON_SCHEMA_SCOPED_MEMBER_PATH_PUSH(MemberPath, TEXT("properties"), true);
-					
-					if (const TOptional<TSharedPtr<FJsonObject>> Schema_Properties = TestAndGetJsonObjectField(
-							Schema, Last(MemberPath), *this);
-						Schema_Properties.IsSet())
-					{
-						UE_JSON_SCHEMA_SCOPED_MEMBER_PATH_PUSH(MemberPath, TEXT("testString"), true);
-						
-						if (const TOptional<TSharedPtr<FJsonObject>> Schema_Properties_TestString = TestAndGetJsonObjectField(
-								*Schema_Properties, Last(MemberPath), *this);
-							Schema_Properties_TestString.IsSet())
-						{
+
 							{
-								UE_JSON_SCHEMA_SCOPED_MEMBER_PATH_PUSH(MemberPath, TEXT("type"), true);
-								
-								if (const TOptional<FString> Schema_Properties_TestString_Type = TestAndGetJsonStringField(
-										*Schema_Properties_TestString, Last(MemberPath), *this);
-									Schema_Properties_TestString_Type.IsSet())
+								const TSharedPtr<FJsonObject>* Property = nullptr;
+								TestTrue("Input", Properties->Get()->TryGetObjectField("objectParam", Property));
+								if (Property)
 								{
-									TestTrue(
-										FString::Printf(TEXT("%s is string."), *MemberPath),
-										*Schema_Properties_TestString_Type == "string");
+									TestString(*this, *Property, "TestObjectDesc", FString(""));
 								}
 							}
 
 							{
-								UE_JSON_SCHEMA_SCOPED_MEMBER_PATH_PUSH(MemberPath, TEXT("description"), true);
-								
-								if (const TOptional<FString> Schema_Properties_TestString_Description = TestAndGetJsonStringField(
-										*Schema_Properties_TestString, Last(MemberPath), *this);
-									Schema_Properties_TestString_Description.IsSet())
+								const TSharedPtr<FJsonObject>* Property = nullptr;
+								TestTrue("Input", Properties->Get()->TryGetObjectField("linearColorParam", Property));
+								if (Property)
 								{
-									TestEqual(
-										FString::Printf(TEXT("%s should be '%s'. Is '%s'."), *MemberPath, 
-											*ExpectedTestStringDescription, **Schema_Properties_TestString_Description),
-										*Schema_Properties_TestString_Description, 
-										ExpectedTestStringDescription); 
+									TestLinearColor(*this, *Property, "TestLinearColorDesc", FLinearColor(0.1f, 0.2f, 0.3f, 0.4f));
 								}
 							}
-						}
-					}
-				});
-			
-			It(
-				"Should generate valid JSON schema from a UStruct's enum UProperty.", 
-				[this]() -> void
-				{
-					FJsonSchemaMemberPath MemberPath(TEXT("[schema]")); // ..mostly for error reporting
-					
-					const TSharedPtr<FJsonObject> Schema = FJsonSchemaGenerator::UStructToJsonSchemaObject<FJsonUtilitiesFakeStruct>();
-					if (!TestIsJsonSchemaValidAndNotEmpty(Schema, *this))
-					{
-						return;
-					}
-						
-					UE_JSON_SCHEMA_SCOPED_MEMBER_PATH_PUSH(MemberPath, TEXT("properties"), true);
-					
-					if (const TOptional<TSharedPtr<FJsonObject>> Schema_Properties = TestAndGetJsonObjectField(
-							Schema, Last(MemberPath), *this);
-						Schema_Properties.IsSet())
-					{
-						UE_JSON_SCHEMA_SCOPED_MEMBER_PATH_PUSH(MemberPath, TEXT("testEnum"), true);
-						
-						if (const TOptional<TSharedPtr<FJsonObject>> Schema_Properties_TestEnum = TestAndGetJsonObjectField(
-								*Schema_Properties, Last(MemberPath), *this);
-							Schema_Properties_TestEnum.IsSet())
-						{
+
 							{
-								UE_JSON_SCHEMA_SCOPED_MEMBER_PATH_PUSH(MemberPath, TEXT("type"), true);
-								
-								if (const TOptional<FString> Schema_Properties_TestEnum_Type = TestAndGetJsonStringField(
-										*Schema_Properties_TestEnum, Last(MemberPath), *this);
-									Schema_Properties_TestEnum_Type.IsSet())
+								const TSharedPtr<FJsonObject>* Property = nullptr;
+								TestTrue("Input", Properties->Get()->TryGetObjectField("colorParam", Property));
+								if (Property)
 								{
-									TestTrue(
-										FString::Printf(TEXT("%s is string."), *MemberPath),
-										*Schema_Properties_TestEnum_Type == "string"); 
+									TestColor(*this, *Property, "TestColorDesc", FColor(10, 20, 30, 40));
 								}
 							}
 
 							{
-								UE_JSON_SCHEMA_SCOPED_MEMBER_PATH_PUSH(MemberPath, TEXT("enum"), true);
-								
-								if (const TOptional<TArray<FString>> Schema_Properties_TestEnum_Enum = TestAndGetJsonStringArrayField(
-										*Schema_Properties_TestEnum, Last(MemberPath), *this);
-									Schema_Properties_TestEnum_Enum.IsSet())
+								const TSharedPtr<FJsonObject>* Property = nullptr;
+								TestTrue("Input", Properties->Get()->TryGetObjectField("rotatorParam", Property));
+								if (Property)
 								{
-									if (TestEqual(
-									FString::Printf(
-										TEXT("%s should have %d values. Has %d values."), 
-											*MemberPath,
-											ExpectedTestEnumValues.Num(),
-											Schema_Properties_TestEnum_Enum->Num()),
-										Schema_Properties_TestEnum_Enum->Num(), ExpectedTestEnumValues.Num())) 
-									{
-										TestTrue(
-											FString::Printf(
-												TEXT("%s should have values ['%s', '%s']. Has values ['%s', '%s']."), 
-												*MemberPath,
-												*ExpectedTestEnumValues[0], *ExpectedTestEnumValues[1],
-												*(*Schema_Properties_TestEnum_Enum)[0], *(*Schema_Properties_TestEnum_Enum)[1]),
-                                     		*(*Schema_Properties_TestEnum_Enum)[0] == ExpectedTestEnumValues[0] && 
-                                     			*(*Schema_Properties_TestEnum_Enum)[1] == ExpectedTestEnumValues[1]);
-									}
+									TestRotator(*this, *Property, "TestRotatorDesc", FRotator(20, 10, 30));
 								}
 							}
 
 							{
-								UE_JSON_SCHEMA_SCOPED_MEMBER_PATH_PUSH(MemberPath, TEXT("description"), true);
-								
-								if (const TOptional<FString> Schema_Properties_TestEnum_Description = TestAndGetJsonStringField(
-										*Schema_Properties_TestEnum, Last(MemberPath), *this);
-									Schema_Properties_TestEnum_Description.IsSet())
+								const TSharedPtr<FJsonObject>* Property = nullptr;
+								TestTrue("Input", Properties->Get()->TryGetObjectField("vectorParam", Property));
+								if (Property)
 								{
-									TestEqual(
-										FString::Printf(TEXT("%s should be '%s'. Is '%s'."), *MemberPath, 
-											*ExpectedTestEnumDescription, **Schema_Properties_TestEnum_Description),
-										*Schema_Properties_TestEnum_Description, 
-										ExpectedTestEnumDescription); 
+									TestVector(*this, *Property, "TestVectorDesc", FVector(1.f, 2.f, 3.f));
 								}
 							}
 						}
 					}
-				});
-			
-			It(
-				"Should generate valid JSON schema from a UStruct - with nested UStruct UProperty.", 
-				[this]() -> void
-				{
-					FJsonSchemaMemberPath MemberPath(TEXT("[schema]")); // ..mostly for error reporting
-					
-					const TSharedPtr<FJsonObject> Schema = FJsonSchemaGenerator::UStructToJsonSchemaObject<FJsonUtilitiesFakeStruct>();
-					if (!TestIsJsonSchemaValidAndNotEmpty(Schema, *this))
-					{
-						return;
-					}
-						
-					UE_JSON_SCHEMA_SCOPED_MEMBER_PATH_PUSH(MemberPath, TEXT("properties"), true);
-					
-					if (const TOptional<TSharedPtr<FJsonObject>> Schema_Properties = TestAndGetJsonObjectField(
-							Schema, Last(MemberPath), *this);
-						Schema_Properties.IsSet())
+
+					const TSharedPtr<FJsonObject>* OutputSchema = nullptr;
+					TestTrue("Output", Schema->TryGetObjectField("outputSchema", OutputSchema));
+					if (OutputSchema)
 					{
-						UE_JSON_SCHEMA_SCOPED_MEMBER_PATH_PUSH(MemberPath, TEXT("testNestedStructA"), true);
-						
-						if (const TOptional<TSharedPtr<FJsonObject>> Schema_Properties_TestStruct = TestAndGetJsonObjectField(
-								*Schema_Properties, Last(MemberPath), *this);
-							Schema_Properties_TestStruct.IsSet())
+						TestEqual("Output", OutputSchema->Get()->Values.Num(), 3);
+						TestBasics(*this, *OutputSchema, "object", "");
+						TestRequired(*this, *OutputSchema, TArray<FString>({ "returnValue" }));
+
+						const TSharedPtr<FJsonObject>* Properties = nullptr;
+						TestTrue("Output", OutputSchema->Get()->TryGetObjectField("properties", Properties));
+						if (Properties)
 						{
+							TestEqual("Output", Properties->Get()->Values.Num(), 1);
+							const TSharedPtr<FJsonObject>* ReturnValue = nullptr;
+							TestTrue("Output", Properties->Get()->TryGetObjectField("returnValue", ReturnValue));
+							if (ReturnValue)
 							{
-								UE_JSON_SCHEMA_SCOPED_MEMBER_PATH_PUSH(MemberPath, TEXT("type"), true);
-								
-								if (const TOptional<FString> Schema_Properties_TestStruct_Type = TestAndGetJsonStringField(
-										*Schema_Properties_TestStruct, Last(MemberPath), *this);
-									Schema_Properties_TestStruct_Type.IsSet())
-								{
-									TestTrue(
-										FString::Printf(TEXT("%s is object."), *MemberPath),
-										*Schema_Properties_TestStruct_Type == "object"); 
-								}
-							}
-							
-							{
-								UE_JSON_SCHEMA_SCOPED_MEMBER_PATH_PUSH(MemberPath, TEXT("description"), true);
-								
-								if (const TOptional<FString> Schema_Properties_TestStruct_Description = TestAndGetJsonStringField(
-										*Schema_Properties_TestStruct, Last(MemberPath), *this);
-									Schema_Properties_TestStruct_Description.IsSet())
-								{
-									TestEqual(
-										FString::Printf(TEXT("%s should be '%s'. Is '%s'."), *MemberPath, 
-											*ExpectedTestNestedStructADescription, **Schema_Properties_TestStruct_Description),
-										*Schema_Properties_TestStruct_Description, 
-										ExpectedTestNestedStructADescription); 
-								}
+								TestFloat(*this, *ReturnValue, "TestReturnDesc");
 							}
 						}
 					}
 				});
 
 			It(
-				"Should generate valid JSON schema from a UClass's UFunction - with simpler defaulted params.", 
+				"Should generate JSON schema from a UFunction with custom structs with defaults.",
 				[this]() -> void
 				{
-					FJsonSchemaMemberPath MemberPath(TEXT("[schema]")); // ..mostly for error reporting
-					
-					const UClass& Class = *UJsonUtilitiesFakeClass::StaticClass();
-
-					static const TCHAR* FunctionName = TEXT("TestFunc");
-					UFunction* FunctionPtr = Class.FindFunctionByName(FunctionName);
-					if (!TestTrue(
-						"Found named function.", 
-						LIKELY(FunctionPtr)))
-					{
-						return;
-					}
-					
+					UFunction* FunctionPtr = UJsonUtilitiesFakeClass::StaticClass()->FindFunctionByName("TestEngineStructFunc");
 					const TSharedPtr<FJsonObject> Schema = FJsonSchemaGenerator::UStructToJsonSchemaObject(FunctionPtr);
-					if (!TestIsJsonSchemaValidAndNotEmpty(Schema, *this))
-					{
-						return;
-					}
-
-#if DEBUG_PRINT_JSONS_TO_LOG
-					PrintJsonToLog(FunctionName, Schema);
-#endif
+					TestTrue("Schema", Schema.IsValid());
+					TestBasics(*this, Schema, "", "TestEngineStructFuncDesc");
 
+					const TSharedPtr<FJsonObject>* InputSchema = nullptr;
+					TestTrue("Input", Schema->TryGetObjectField("inputSchema", InputSchema));
+					if (InputSchema)
 					{
-						UE_JSON_SCHEMA_SCOPED_MEMBER_PATH_PUSH(MemberPath, TEXT("inputSchema"), true);
-						
-						if (const TOptional<TSharedPtr<FJsonObject>> Schema_InputSchema = TestAndGetJsonObjectField(
-								Schema, Last(MemberPath), *this);
-							Schema_InputSchema.IsSet())
+						TestBasics(*this, *InputSchema, "object", "");
+						TestRequired(*this, *InputSchema, TArray<FString>());
+						const TSharedPtr<FJsonObject>* Properties = nullptr;
+						TestTrue("Input", InputSchema->Get()->TryGetObjectField("properties", Properties));
+						if (Properties)
 						{
-							UE_JSON_SCHEMA_SCOPED_MEMBER_PATH_PUSH(MemberPath, TEXT("properties"), true);
-							
-							if (const TOptional<TSharedPtr<FJsonObject>> Schema_InputSchema_Properties = TestAndGetJsonObjectField(
-									*Schema_InputSchema, Last(MemberPath), *this);
-								Schema_InputSchema_Properties.IsSet())
+							const TSharedPtr<FJsonObject>* Transform = nullptr;
+							TestTrue("Input", Properties->Get()->TryGetObjectField("transform", Transform));
+							if (Transform)
 							{
-								UE_JSON_SCHEMA_SCOPED_MEMBER_PATH_PUSH(MemberPath, TEXT("testFuncParamFloat"), true);
-								
-								if (const TOptional<TSharedPtr<FJsonObject>> Schema_InputSchema_Properties_TestFuncParamFloat = TestAndGetJsonObjectField(
-										*Schema_InputSchema_Properties, Last(MemberPath), *this);
-									Schema_InputSchema_Properties_TestFuncParamFloat.IsSet())
-								{
-									{
-										UE_JSON_SCHEMA_SCOPED_MEMBER_PATH_PUSH(MemberPath, TEXT("type"), true);
-										
-										if (const TOptional<FString> Schema_InputSchema_Properties_TestFuncParamFloat_Type = TestAndGetJsonStringField(
-												*Schema_InputSchema_Properties_TestFuncParamFloat, Last(MemberPath), *this);
-											Schema_InputSchema_Properties_TestFuncParamFloat_Type.IsSet())
-										{
-											TestTrue(
-												FString::Printf(TEXT("%s is number."), *MemberPath),
-												*Schema_InputSchema_Properties_TestFuncParamFloat_Type == "number"); 
-										}
-									}
-
-									{
-										UE_JSON_SCHEMA_SCOPED_MEMBER_PATH_PUSH(MemberPath, TEXT("description"), true);
-										
-										if (const TOptional<FString> Schema_InputSchema_Properties_TestFuncParamFloat_Description = TestAndGetJsonStringField(
-												*Schema_InputSchema_Properties_TestFuncParamFloat, Last(MemberPath), *this);
-											Schema_InputSchema_Properties_TestFuncParamFloat_Description.IsSet())
-										{
-											TestEqual(
-												FString::Printf(TEXT("%s should be '%s'. Is '%s'."), *MemberPath, 
-													*ExpectedTestFuncParamFloatDescription, **Schema_InputSchema_Properties_TestFuncParamFloat_Description),
-												*Schema_InputSchema_Properties_TestFuncParamFloat_Description, 
-												ExpectedTestFuncParamFloatDescription); 
-										}
-									}
-								}
+								TestTransform(*this, *Transform, "TestEngineStructFuncParamTransformDesc",
+									FTransform(FQuat(0.1, 0.2, 0.3, 0.4), FVector(1, 2, 3), FVector(10, 20, 30)));
 							}
 						}
 					}
-					
+
+					const TSharedPtr<FJsonObject>* OutputSchema = nullptr;
+					TestTrue("Output", Schema->TryGetObjectField("outputSchema", OutputSchema));
+					if (OutputSchema)
 					{
-						UE_JSON_SCHEMA_SCOPED_MEMBER_PATH_PUSH(MemberPath, TEXT("outputSchema"), true);
-						
-						if (const TOptional<TSharedPtr<FJsonObject>> Schema_OutputSchema = TestAndGetJsonObjectField(
-								Schema, Last(MemberPath), *this);
-							Schema_OutputSchema.IsSet())
+						TestEqual("Output", OutputSchema->Get()->Values.Num(), 3);
+						TestBasics(*this, *OutputSchema, "object", "");
+						TestRequired(*this, *OutputSchema, TArray<FString>({ "returnValue" }));
+						const TSharedPtr<FJsonObject>* Properties = nullptr;
+						TestTrue("Output", OutputSchema->Get()->TryGetObjectField("properties", Properties));
+						if (Properties)
 						{
-							UE_JSON_SCHEMA_SCOPED_MEMBER_PATH_PUSH(MemberPath, TEXT("properties"), true);
-							
-							if (const TOptional<TSharedPtr<FJsonObject>> Schema_OutputSchema_Properties = TestAndGetJsonObjectField(
-									*Schema_OutputSchema, Last(MemberPath), *this);
-								Schema_OutputSchema_Properties.IsSet())
+							TestEqual("Output", Properties->Get()->Values.Num(), 1);
+							const TSharedPtr<FJsonObject>* ReturnValue = nullptr;
+							TestTrue("Output", Properties->Get()->TryGetObjectField("returnValue", ReturnValue));
+							if (ReturnValue)
 							{
-								UE_JSON_SCHEMA_SCOPED_MEMBER_PATH_PUSH(MemberPath, TEXT("returnValue"), true);
-								
-								if (const TOptional<TSharedPtr<FJsonObject>> Schema_OutputSchema_Properties_TestReturnValueObject = TestAndGetJsonObjectField(
-										*Schema_OutputSchema_Properties, Last(MemberPath), *this);
-									Schema_OutputSchema_Properties_TestReturnValueObject.IsSet())
+								TestEqual("Output", ReturnValue->Get()->Values.Num(), 4);
+								TestBasics(*this, *ReturnValue, "object", "TestEngineStructFuncReturnValueDesc");
+								TestRequired(*this, *ReturnValue, TArray<FString>({ "testNestedStructFloat" }));
+								TestTrue("Output", ReturnValue->Get()->TryGetObjectField("properties", Properties));
+								if (Properties)
 								{
+									TestEqual("Output", Properties->Get()->Values.Num(), 1);
+									const TSharedPtr<FJsonObject>* FloatProperty = nullptr;
+									TestTrue("Output", Properties->Get()->TryGetObjectField("testNestedStructFloat", FloatProperty));
+									if (FloatProperty)
 									{
-										UE_JSON_SCHEMA_SCOPED_MEMBER_PATH_PUSH(MemberPath, TEXT("type"), true);
-										
-										if (TOptional<FString> Schema_OutputSchema_Properties_TestFuncReturnValue_Type = TestAndGetJsonStringField(
-												*Schema_OutputSchema_Properties_TestReturnValueObject, Last(MemberPath), *this);
-											Schema_OutputSchema_Properties_TestFuncReturnValue_Type.IsSet())
-										{
-											TestTrue(
-												FString::Printf(TEXT("%s type is number."), *MemberPath),
-												*Schema_OutputSchema_Properties_TestFuncReturnValue_Type == "number"); 
-										}
-									}
-
-									{
-										UE_JSON_SCHEMA_SCOPED_MEMBER_PATH_PUSH(MemberPath, TEXT("description"), true);
-										
-										if (const TOptional<FString> Schema_OutputSchema_Properties_TestFuncReturnValue_Description = TestAndGetJsonStringField(
-												*Schema_OutputSchema_Properties_TestReturnValueObject, Last(MemberPath), *this);
-											Schema_OutputSchema_Properties_TestFuncReturnValue_Description.IsSet())
-										{
-											TestEqual(
-												FString::Printf(TEXT("%s should be '%s'. Is '%s'."), *MemberPath, 
-													*ExpectedTestFuncReturnValueDescription, **Schema_OutputSchema_Properties_TestFuncReturnValue_Description),
-												*Schema_OutputSchema_Properties_TestFuncReturnValue_Description, 
-												ExpectedTestFuncReturnValueDescription); 
-										}
+										TestFloat(*this, *FloatProperty, "TestNestedFloatDesc", -301.f, 302.f);
 									}
 								}
 							}
 						}
 					}
-
-					{
-						UE_JSON_SCHEMA_SCOPED_MEMBER_PATH_PUSH(MemberPath, TEXT("description"), true);
-						
-						// In this case, overall schema description should be function description.
-						if (const TOptional<FString> Schema_Description = TestAndGetJsonStringField(
-								Schema, Last(MemberPath), *this);
-							Schema_Description.IsSet())
-						{
-							TestEqual(
-								FString::Printf(TEXT("%s should be '%s'. Is '%s'."), *MemberPath, 
-									*ExpectedTestFuncDescription, **Schema_Description),
-								*Schema_Description, 
-								ExpectedTestFuncDescription); 
-						}
-					}
 				});
 
 			It(
-				"Should generate valid JSON schema from a UClass's UFunction - with Engine struct defaulted params.", 
+				"Should not generate JSON schema from an UFunction with non-Json-able params.",
 				[this]() -> void
-				{
-					FJsonSchemaMemberPath MemberPath(TEXT("[schema]")); // ..mostly for error reporting
-					
-					const UClass& Class = *UJsonUtilitiesFakeClass::StaticClass();
-
-					static const TCHAR* FunctionName = TEXT("TestEngineStructFunc");
-					UFunction* FunctionPtr = Class.FindFunctionByName(FunctionName);
-					if (!TestTrue(
-						"Found named function.", 
-						LIKELY(FunctionPtr)))
+				{									
+					UFunction* FunctionPtr = UJsonUtilitiesFakeClass::StaticClass()->FindFunctionByName("TestIllegalFunc");
+					if (!TestTrue("Found named function.", LIKELY(FunctionPtr)))
 					{
 						return;
 					}
-					
-					const TSharedPtr<FJsonObject> Schema = FJsonSchemaGenerator::UStructToJsonSchemaObject(FunctionPtr);
-					if (!TestIsJsonSchemaValidAndNotEmpty(Schema, *this))
-					{
-						return;
-					}
-
-#if DEBUG_PRINT_JSONS_TO_LOG
-					PrintJsonToLog(FunctionName, Schema);
-#endif
+				
+					AddExpectedError(TEXT(".*unhandled during Json schema generation.*"),
+						EAutomationExpectedErrorFlags::Contains, INDEX_NONE);
 						
-					// TODO - Test is incomplete.
-				});
+					const TSharedPtr<FJsonObject> Schema = FJsonSchemaGenerator::UStructToJsonSchemaObject(FunctionPtr);
+					TestFalse("Schema is invalid.", Schema.IsValid());					
+			});
 
 			It(
-				"Should generate invalid JSON schema from a UClass's UFunction - with illegal non-Json-able params.", 
+				"Should support custom callbacks.",
 				[this]() -> void
 				{
-					FJsonSchemaMemberPath MemberPath(TEXT("[schema]")); // ..mostly for error reporting
-				
-					const UClass& Class = *UJsonUtilitiesFakeClass::StaticClass();
-
-					static const TCHAR* FunctionName = TEXT("TestIllegalFunc");
-					UFunction* FunctionPtr = Class.FindFunctionByName(FunctionName);
-					if (!TestTrue(
-						"Found named function.", 
-						LIKELY(FunctionPtr)))
+					UFunction* FunctionPtr = UJsonUtilitiesFakeClass::StaticClass()->FindFunctionByName("TestEngineCustomizedFunc");
+					FJsonSchemaPropertyFilter::CustomCallback CustomCallback;
+					CustomCallback.BindStatic(TestCustomCallback);
+					FJsonSchemaPropertyFilter PropertyFilter;
+					PropertyFilter.CustomCb = &CustomCallback;
+					const TSharedPtr<FJsonObject> Schema = FJsonSchemaGenerator::UStructToJsonSchemaObject(FunctionPtr, PropertyFilter);
+					TestTrue("Schema", Schema.IsValid());
+					TestBasics(*this, Schema, "", "TestCustomizedFuncDesc");
+					const TSharedPtr<FJsonObject>* InputSchema = nullptr;
+					TestTrue("Input", Schema->TryGetObjectField("inputSchema", InputSchema));
+					if (InputSchema)
 					{
-						return;
+						TestBasics(*this, *InputSchema, "object", "");
+						TestRequired(*this, *InputSchema, TArray<FString>());
+						const TSharedPtr<FJsonObject>* Properties = nullptr;
+						TestTrue("Input", InputSchema->Get()->TryGetObjectField("properties", Properties));
+						if (Properties)
+						{
+							const TSharedPtr<FJsonObject>* SoftPath = nullptr;
+							TestTrue("Input", Properties->Get()->TryGetObjectField("softPath", SoftPath));
+							if (SoftPath)
+							{
+								TestString(*this, *SoftPath, "TestSoftPathDesc", FString("/Engine/StaticMesh.StaticMesh"));
+							}
+						}
 					}
-				
-					AddExpectedError(TEXT(".*Illegal property type for Json.*"),
-						EAutomationExpectedErrorFlags::Contains, INDEX_NONE);
-						
-					const TSharedPtr<FJsonObject> Schema = FJsonSchemaGenerator::UStructToJsonSchemaObject(FunctionPtr);
-					TestIsJsonSchemaInvalid(Schema, *this);
-			});
+				});
 	});
 }
 
diff --git a/Engine/Plugins/Tests/EditorTests/Source/EditorTests/Private/UnrealEd/JsonUtilitiesTests.h b/Engine/Plugins/Tests/EditorTests/Source/EditorTests/Private/UnrealEd/JsonUtilitiesTests.h
--- a/Engine/Plugins/Tests/EditorTests/Source/EditorTests/Private/UnrealEd/JsonUtilitiesTests.h
+++ b/Engine/Plugins/Tests/EditorTests/Source/EditorTests/Private/UnrealEd/JsonUtilitiesTests.h
@@ -27,8 +27,8 @@ enum class EJsonUtilitiesFakeEnum : uint8
 /**
  * TestNestedStructADesc
  */
-USTRUCT() 
-struct FJsonUtilitiesFakeNestedStructA
+USTRUCT(BlueprintType) 
+struct FJsonUtilitiesFakeNestedStruct
 {
 	GENERATED_BODY()
 	
@@ -36,59 +36,9 @@ struct FJsonUtilitiesFakeNestedStructA
 	 * TestNestedFloatDesc
 	 */ 
 	UPROPERTY(EditAnywhere, Category="Nested", meta=(ClampMin="-301", ClampMax="302"))
-	float TestNestedStructAFloat = 3.0f;
+	float TestNestedStructFloat = 3.0f;
 };
 
-/**
- * TestNestedStructBDesc
- */
-USTRUCT() 
-struct FJsonUtilitiesFakeNestedStructB
-{
-	GENERATED_BODY()
-	
-	/**
-	 * TestNestedFloatDesc
-	 */ 
-	UPROPERTY(EditAnywhere, Category="Nested", meta=(ClampMin="-401", ClampMax="402"))
-	float TestNestedStructBFloat = 4.0f;
-};
-
-/**
- * TestNestedStructCDesc
- */
-USTRUCT() 
-struct FJsonUtilitiesFakeNestedStructC
-{
-	GENERATED_BODY()
-	
-	/**
-	 * TestNestedFloatDesc
-	 */ 
-	UPROPERTY(EditAnywhere, Category="Nested", meta=(ClampMin="-501", ClampMax="502"))
-	float TestNestedStructCFloat = 5.0f;
-};
-
-/**
- * TestNestedStructDDesc
- */
-USTRUCT() 
-struct FJsonUtilitiesFakeNestedStructD
-{
-	GENERATED_BODY()
-	
-	/**
-	 * TestNestedStructBDesc
-	 */ 
-	UPROPERTY(EditAnywhere, Category="Nested")
-	FJsonUtilitiesFakeNestedStructB TestNestedStructB;
-	
-	/**
-	 * TestNestedStructCDesc
-	 */ 
-	UPROPERTY(EditAnywhere, Category="Nested")
-	FJsonUtilitiesFakeNestedStructC TestNestedStructC;
-};
 
 /**
  * TestStructDesc
@@ -122,109 +72,106 @@ struct FJsonUtilitiesFakeStruct
 	UPROPERTY(EditAnywhere, Category="Main")
 	FString TestString = "DefaultString";
 	
+	/**
+	 * TestNameDesc
+	 */
+	UPROPERTY(EditAnywhere, Category = "Main")
+	FName TestName = "DefaultName";
+
 	/**
 	 * TestEnumDesc
 	 */ 
 	UPROPERTY(EditAnywhere, Category="Main")
 	EJsonUtilitiesFakeEnum TestEnum = EJsonUtilitiesFakeEnum::One;
-	
+
 	/**
-	 * TestNestedStructDDesc
-	 */ 
-	UPROPERTY(EditAnywhere, Category="Main")
-	FJsonUtilitiesFakeNestedStructD TestNestedStructD;
+	 * TestObjectDesc
+	 */
+	UPROPERTY(EditAnywhere, Category = "Main")
+	TObjectPtr<UObject> TestObject;
+
+	/**
+	 * TestClassDesc
+	 */
+	UPROPERTY(EditAnywhere, Category = "Main")
+	TObjectPtr<UClass> TestClass;
 	
 	/**
-	 * TestNestedStructADesc
+	 * TestNestedStructDesc
 	 */ 
 	UPROPERTY(EditAnywhere, Category="Main")
-	FJsonUtilitiesFakeNestedStructA TestNestedStructA;
-};
-
-/**
- * FFunctionParamStruct
- * Used for testing how structs end up working with function parameters in function schemas.
- */
-USTRUCT(BlueprintType)
-struct FFunctionParamStruct
-{
-	GENERATED_BODY()
-	
-	UPROPERTY(Category="Main", BlueprintReadWrite, EditAnywhere)
-	FString AString = "";
-
-	UPROPERTY(Category="Main", BlueprintReadWrite, EditAnywhere)
-	FString BString = "";
+	FJsonUtilitiesFakeNestedStruct TestNestedStruct;
 };
 
 /**
  * TestClassDesc
  */
 UCLASS(BlueprintType)
-class UJsonUtilitiesFakeClass : public UObject, public FJsonUtilitiesFakeStruct
+class UJsonUtilitiesFakeClass : public UObject
 {
 	GENERATED_BODY()
+
+	/**
+	 * TestFloatDesc
+	 */
+	UPROPERTY(EditAnywhere, Category = "Main", meta = (ClampMin = "-101.0", ClampMax = "102.0"))
+	float TestFloat = 1.0f;
 	
 	/**
-	 * TestFuncDesc
-	 * @param TestFuncParamFakeStruct TestFuncParamFakeStructDesc
-	 * @param bTestFuncParamBool TestFuncParamBoolDesc
-	 * @param TestFuncParamFloat TestFuncParamFloatDesc
-	 * @param TestFuncParamInt TestFuncParamIntDesc
-	 * @param TestFuncParamEnum TestFuncParamEnumDesc
-	 * @return TestFuncReturnValueDesc
+	 * TestFuncDesc	 
+	 * @param RequiredParam TestRequiredDesc
+	 * @param BoolParam TestBoolDesc
+	 * @param FloatParam TestFloatDesc
+	 * @param IntParam TestIntDesc
+	 * @param EnumParam TestEnumDesc
+	 * @param StringParam TestStringDesc
+	 * @param NameParam TestNameDesc
+	 * @param ObjectParam TestObjectDesc	 
+	 * @param LinearColorParam TestLinearColorDesc
+	 * @param ColorParam TestColorDesc
+	 * @param RotatorParam TestRotatorDesc
+	 * @param VectorParam TestVectorDesc
+	 * @return TestReturnDesc
 	 */
 	UFUNCTION(BlueprintCallable, Category="Main") 
-	float TestFunc(const FJsonUtilitiesFakeStruct TestFuncParamFakeStruct, const bool bTestFuncParamBool = true, 
-		const float TestFuncParamFloat = 1.23f, const int32 TestFuncParamInt = 123,
-		const EJsonUtilitiesFakeEnum TestFuncParamEnum = EJsonUtilitiesFakeEnum::One)
+	float TestFunc(int RequiredParam, bool BoolParam = true, float FloatParam = 1.23f,
+		int32 IntParam = 123, EJsonUtilitiesFakeEnum EnumParam = EJsonUtilitiesFakeEnum::One,
+		const FString& StringParam = "foo", const FName& NameParam = FName("bar"), const UObject* ObjectParam = nullptr,
+		FLinearColor LinearColorParam = FLinearColor(0.1f, 0.2f, 0.3f, 0.4f), FColor ColorParam = FColor(10, 20, 30, 40),
+		FRotator RotatorParam = FRotator(20, 10, 30), FVector VectorParam = FVector(1.f, 2.f, 3.f))
 	{
 		// This function is for testing schema generation for a function with 
 		// simpler parameters.
-		
 		return 0.0f;
 	}
 	
 	/**
-	 * TestEngineStructFuncDesc
-	 * @param SoftObjectPath TestEngineStructFuncParamSoftObjectPathDesc 
-	 * @param SoftClassPath TestEngineStructFuncParamSoftClassPathDesc 
-	 * @param FunctionParamStruct TestEngineStructFuncParamFunctionParamStructDesc 
-	 * @param LinearColor TestEngineStructFuncParamLinearColorDesc 
-	 * @param Color TestEngineStructFuncParamColorDesc 
-	 * @param Quat TestEngineStructFuncParamQuatDesc 
-	 * @param Rotator TestEngineStructFuncParamRotatorDesc 
-	 * @param Transform TestEngineStructFuncParamTransformDesc 
-	 * @param Vector TestEngineStructFuncParamVectorDesc
+	 * TestEngineStructFuncDesc	  
+	 * @param Transform TestEngineStructFuncParamTransformDesc 	 
 	 * @return TestEngineStructFuncReturnValueDesc
 	 */
-	UFUNCTION(BlueprintCallable, Category="Main", meta=(
-		CPP_Default_SoftObjectPath="/Game/MyFolder/MyAsset.MyAsset", 
-		CPP_Default_SoftClassPath="/Game/MyFolder/MyActorBP.MyMyActorBP_C",
-		CPP_Default_FunctionParamStruct="(AString=\"AAA\",BString=\"BBB\")",
-		CPP_Default_LinearColor="(R=0.1,G=0.2,B=0.3,A=0.4)",
-		CPP_Default_Color="(R=10,G=20,B=30,A=40)",
-		CPP_Default_Quat="(X=0.1,Y=0.2,Z=0.3,W=0.4)",
-		CPP_Default_Rotator="(Roll=30,Yaw=10,Pitch=20)",
+	UFUNCTION(BlueprintCallable, Category="Main", meta=(		
 		CPP_Default_Transform="(Rotation=(X=0.1,Y=0.2,Z=0.3,W=0.4),Translation=(X=1,Y=2,Z=3),Scale3D=(X=10,Y=20,Z=30))"))
-	float TestEngineStructFunc(
-		FSoftObjectPath SoftObjectPath, 
-		FSoftClassPath SoftClassPath,
-		FFunctionParamStruct FunctionParamStruct, 
-		FLinearColor LinearColor, FColor Color,
-		FQuat Quat, FRotator Rotator, FTransform Transform, 
-		FVector Vector = FVector::UpVector)
+	FJsonUtilitiesFakeNestedStruct TestEngineStructFunc(FTransform Transform)
 	{
 		// This function is for testing schema generation for a function with
 		// more complex struct parameters, with defaults.
 		// UHT does allow some structs of be initialized in the declaration, like FVector.
 		// UHT does not allow most other structs to be initialized in the declaration, like 
-		// FSoftObjectPath, or arbitrary structs like FFunctionParamStruct.. and for those we 
-		// need to use CPP_Default_XXX. 
+		// FTransform or arbitrary structs and for those we  need to use CPP_Default_XXX. 
 		// The real purpose of this function is for testing how these defaults end up looking in
 		// the generated JSON schema.
 
-		return 0.0f;
+		return FJsonUtilitiesFakeNestedStruct();
+	}
+
+	/**
+	 * TestCustomizedFuncDesc
+	 * @param SoftPath TestSoftPathDesc
+	 */
+	UFUNCTION(BlueprintCallable, Category = "Main", meta = (CPP_Default_SoftPath = "/Engine/StaticMesh.StaticMesh"))
+	void TestEngineCustomizedFunc(FSoftClassPath SoftPath)
+	{
 	}
 	
 	/**
@@ -232,12 +179,7 @@ class UJsonUtilitiesFakeClass : public UObject, public FJsonUtilitiesFakeStruct
 	 * @return TestIllegalFuncReturnValueDesc
 	 */
 	UFUNCTION(BlueprintCallable, Category="Main")
-	UObject* TestIllegalFunc(
-		int32 IntParamA = 1, 
-		UObject* ObjectParam = nullptr, 
-		int32 IntParamB = 2, 
-		UClass* ClassParam = nullptr, 
-		int32 IntParamC = 3)
+	UObject* TestIllegalFunc(TMap<FVector, float>& Arg)
 	{
 		// This function is for testing schema generation for a function with some parameters that 
 		// can't be represented in the JSON schema.
diff --git a/Engine/Source/Editor/JsonUtilitiesEditor/Private/JsonSchema/JsonSchemaGeneratorEditor.cpp b/Engine/Source/Editor/JsonUtilitiesEditor/Private/JsonSchema/JsonSchemaGeneratorEditor.cpp
--- a/Engine/Source/Editor/JsonUtilitiesEditor/Private/JsonSchema/JsonSchemaGeneratorEditor.cpp
+++ b/Engine/Source/Editor/JsonUtilitiesEditor/Private/JsonSchema/JsonSchemaGeneratorEditor.cpp
@@ -237,7 +237,7 @@ namespace FJsonSchemaMetadataPropertyDefaults
 				}
 				else if (ScriptStruct == TBaseStructure<FRotator>::Get())
 				{
-					NumberValueNames = {"yaw", "pitch", "roll"}; 
+					NumberValueNames = {"pitch", "yaw", "roll"}; 
 				}
 				else if (ScriptStruct == TBaseStructure<FLinearColor>::Get() ||
 					ScriptStruct == TBaseStructure<FColor>::Get())
@@ -281,7 +281,7 @@ namespace FJsonSchemaMetadataPropertyDefaults
 
 	void SetFunctionPropertyDefaultInPropertyMetadata(
 		const UFunction* PropertyOwnerFunction, const FProperty* Property, 
-		const TSharedRef<FJsonObject>& OutPropertyMetadata)
+		const FJsonSchemaPropertyFilter& PropertyFilter, const TSharedRef<FJsonObject>& OutPropertyMetadata)
 	{
 		if (Property->HasAnyPropertyFlags(CPF_ReturnParm))
 		{
@@ -295,6 +295,14 @@ namespace FJsonSchemaMetadataPropertyDefaults
 		UEdGraphSchema_K2::FindFunctionParameterDefaultValue(PropertyOwnerFunction, Property, ParameterDefaultString);
 		if (!ParameterDefaultString.IsEmpty())
 		{
+			if (PropertyFilter.CustomCb && PropertyFilter.CustomCb->IsBound())
+			{
+				if (PropertyFilter.CustomCb->Execute(Property, ParameterDefaultString, OutPropertyMetadata))
+				{
+					return;
+				}
+			}
+
 			if (const FEnumProperty* EnumProperty = CastField<FEnumProperty>(Property))
 			{
 				if (const UEnum* Enum = EnumProperty->GetEnum())
@@ -379,7 +387,7 @@ namespace FJsonSchemaMetadataPropertyDefaults
 			}
 			else if (Property->IsA<FObjectProperty>())
 			{
-				// Do not attempt to set defaults for this property type.
+				OutPropertyMetadata->SetStringField(TEXT("default"), ParameterDefaultString == "None" ? "" : ParameterDefaultString);				
 			}
 			else
 			{
@@ -490,7 +498,7 @@ public:
 		// If function, add any parameter defaults.
 		if (PropertyOwnerFunction)
 		{
-			FJsonSchemaMetadataPropertyDefaults::SetFunctionPropertyDefaultInPropertyMetadata(PropertyOwnerFunction, Property, PropertyMetadata);
+			FJsonSchemaMetadataPropertyDefaults::SetFunctionPropertyDefaultInPropertyMetadata(PropertyOwnerFunction, Property, PropertyFilter, PropertyMetadata);
 		}
 		
 		if (PropertyMetadata->Values.Num())
@@ -566,9 +574,7 @@ private:
 		}
 		else if (const FMapProperty* MapProperty = CastField<FMapProperty>(&Property))
 		{
-			if (ensureMsgf(MapProperty->KeyProp->IsA<FStrProperty>(), 
-				TEXT("GetFPropertyInnerStruct() - Unsupported map key type for %s. JSON-compatible map keys must be FString. Map will be skipped."), 
-				PropertyPath))
+			if (MapProperty->KeyProp->IsA<FStrProperty>())
 			{
 				if (const FStructProperty* ValueStructProperty = CastField<FStructProperty>(MapProperty->ValueProp))
 				{
diff --git a/Engine/Source/Runtime/JsonUtilities/Private/JsonSchema/JsonSchemaGenerator.cpp b/Engine/Source/Runtime/JsonUtilities/Private/JsonSchema/JsonSchemaGenerator.cpp
--- a/Engine/Source/Runtime/JsonUtilities/Private/JsonSchema/JsonSchemaGenerator.cpp
+++ b/Engine/Source/Runtime/JsonUtilities/Private/JsonSchema/JsonSchemaGenerator.cpp
@@ -41,47 +41,6 @@ namespace
 		
 		return (bSuccess ? OutString : TEXT(""));
 	}
-	
-	bool CanPropertyBeDirectlyRepresentedInJson(const FValidConstFProperty& Property)
-	{
-		return !(
-			// Object, Classes, interfaces (pointer-like types)..
-			Property->IsA<FObjectPropertyBase>() || 
-			Property->IsA<FClassProperty>() ||
-			Property->IsA<FSoftClassProperty>() ||
-			Property->IsA<FInterfaceProperty>() ||	
-			Property->IsA<FFieldPathProperty>() ||
-			// Delegates..
-			Property->IsA<FDelegateProperty>() || 
-			Property->IsA<FMulticastDelegateProperty>() || 
-			Property->IsA<FMulticastInlineDelegateProperty>() || 
-			Property->IsA<FMulticastSparseDelegateProperty>());
-	}
-	
-	bool CanStructBeDirectlyRepresentedInJson(const FValidConstUStruct& Struct, FJsonSchemaEditorMetadata& EditorMetadata)
-	{
-		bool bCanAllPropertiesBeAcceptedForJson = true;
-		for (const FProperty* Property : TFieldRange<FProperty>(Struct))
-		{
-			if (!CanPropertyBeDirectlyRepresentedInJson(Property))
-			{
-				const FString InputOrOutputParamString = (Property->HasAnyPropertyFlags(CPF_ReturnParm) ? 
-					TEXT("return") : TEXT("input"));
-				
-				UE_LOG(LogJson, Error, TEXT(
-					"Illegal property type for Json - " // Warning, this string necessary during testing, do not change.
-					"UFunction '%s' %s param '%s' is type %s, and cannot be directly represented in the JSON schema."),
-					*Struct->GetPathName(), 
-					*InputOrOutputParamString, 
-					*Property->GetName(), 
-					*Property->GetCPPType());
-				
-				bCanAllPropertiesBeAcceptedForJson = false;
-			}
-		}
-		
-		return bCanAllPropertiesBeAcceptedForJson;
-	}
 }
 
 //
@@ -124,7 +83,8 @@ public:
 			else if (StackElement.VisitorTarget.IsType<FValidConstFProperty>())
 			{
 				if (!VisitFProperty(StackElement.VisitorTarget.Get<FValidConstFProperty>(), 
-					StackElement.OutputSchema.ToSharedRef(), StackElement.MemberPath, StackElement.PropertyFilter))
+					StackElement.OutputSchema.ToSharedRef(), StackElement.MemberPath, StackElement.PropertyFilter,
+					StackElement.Flags))
 				{
 					return nullptr;
 				}
@@ -135,63 +95,83 @@ public:
 	}
 	
 	bool VisitFProperty(const FValidConstFProperty Property, const TSharedRef<FJsonObject>& OutputSchema, 
-		const FJsonSchemaMemberPath& MemberPath, const FJsonSchemaPropertyFilter& PropertyFilter)
+		const FJsonSchemaMemberPath& MemberPath, const FJsonSchemaPropertyFilter& PropertyFilter,
+		const EVisitorStackElementFlags Flags)
 	{
 		EditorMetadata.CurrentPropertyMemberPath = MemberPath;
 		
-		if (const FOptionalProperty* OptionalProperty = CastField<FOptionalProperty>(Property))
+		bool ValidProperty = false;
+		if (PropertyFilter.CustomCb && PropertyFilter.CustomCb->IsBound())
+		{
+			ValidProperty = PropertyFilter.CustomCb->Execute(Property, "", OutputSchema);
+		}
+
+		if (ValidProperty)
 		{
-			BuildJsonForOptionalProperty(*OptionalProperty, PropertyFilter, OutputSchema);
+			// Do nothing because the CustomCb handled it.
+		}
+		else if (const FOptionalProperty* OptionalProperty = CastField<FOptionalProperty>(Property))
+		{
+			ValidProperty = BuildJsonForOptionalProperty(*OptionalProperty, PropertyFilter, OutputSchema);
 		}
 		else if (const FStructProperty* StructProperty = CastField<FStructProperty>(Property))
 		{
-			BuildJsonForStructProperty(*StructProperty, PropertyFilter, OutputSchema);
+			ValidProperty = BuildJsonForStructProperty(*StructProperty, PropertyFilter, OutputSchema);
 		}
+		else if (const FObjectProperty* ObjectProperty = CastField<FObjectProperty>(Property))
+		{
+			ValidProperty = BuildJsonForObjectProperty(*ObjectProperty, OutputSchema);
+		}		
 		else if (const FArrayProperty* ArrayProperty = CastField<FArrayProperty>(Property))
 		{
-			BuildJsonForArrayProperty(*ArrayProperty, PropertyFilter, OutputSchema);
+			ValidProperty = BuildJsonForArrayProperty(*ArrayProperty, PropertyFilter, OutputSchema);
 		}
 		else if (const FSetProperty* SetProperty = CastField<FSetProperty>(Property))
 		{
-			BuildJsonForSetProperty(*SetProperty, PropertyFilter, OutputSchema);
+			ValidProperty = BuildJsonForSetProperty(*SetProperty, PropertyFilter, OutputSchema);
 		}
 		else if (const FMapProperty* MapProperty = CastField<FMapProperty>(Property))
 		{
-			BuildJsonForMapProperty(*MapProperty, PropertyFilter, OutputSchema);
+			ValidProperty = BuildJsonForMapProperty(*MapProperty, PropertyFilter, OutputSchema);
 		}
 		else if (const FEnumProperty* EnumProperty = CastField<FEnumProperty>(Property))
 		{
-			if (EnumProperty->GetEnum())
-			{
-				BuildJsonForEnumDefinition(*EnumProperty->GetEnum(), OutputSchema);
-			}
+			ValidProperty = BuildJsonForEnumDefinition(EnumProperty->GetEnum(), OutputSchema);							
 		}
 		else if (const FNumericProperty* NumericProperty = CastField<FNumericProperty>(Property))
 		{
-			BuildJsonForNumericProperty(*NumericProperty, OutputSchema);
+			ValidProperty = BuildJsonForNumericProperty(*NumericProperty, OutputSchema);
 		}
 		else if (CastField<FBoolProperty>(Property))
 		{
-			BuildJsonForBoolProperty(OutputSchema);
+			ValidProperty = BuildJsonForBoolProperty(OutputSchema);
 		}
 		else if (CastField<FStrProperty>(Property) || 
 			CastField<FTextProperty>(Property) || 
 			CastField<FNameProperty>(Property) ||
 			CastField<FUtf8StrProperty>(Property))
 		{
-			BuildJsonForStringProperty(OutputSchema);
+			ValidProperty = BuildJsonForStringProperty(OutputSchema);
 		}
-		else
+
+		if (!ValidProperty)		
 		{
-			UE_LOG(LogJson, Error, TEXT(
-				"Property \"%s\" type %s unhandled during Json schema generation."),
-				*Property->GetName(), 
-				*Property->GetCPPType());
-			
-			return false;
+			// This means we've hit an unsupported property, which is either an error or a warning depending on
+			// whether all properties are required or not. Generally all properties are only required for UFunctions.
+			const FString Message = FString::Printf(TEXT("Property \"%s\" type %s unhandled during Json schema generation."),
+				*Property->GetName(), *Property->GetCPPType());
+			if ((static_cast<uint8>(Flags) & static_cast<uint8>(EVisitorStackElementFlags::RequireAllProperties)) == 0)
+			{
+				UE_LOG(LogJson, Warning, TEXT("%s"), *Message);
+				return true;
+			}
+			else
+			{
+				UE_LOG(LogJson, Error, TEXT("%s"), *Message);
+				return false;
+			}
 		}
 
-
 		// Add all property metadata, if any.
 		if (const TSharedPtr<FJsonObject> PropertyMetadataJson = EditorMetadata.GetPropertyMetadataForCurrentPropertyMemberPath())
 		{
@@ -213,12 +193,6 @@ public:
 
 		const bool bUseJsonRpcFormat = Struct->IsA(UFunction::StaticClass());
 		
-		// If using JSON RPC format, check that all properties can be accepted.
-		if (bUseJsonRpcFormat && !CanStructBeDirectlyRepresentedInJson(Struct, EditorMetadata))
-		{
-			return false;
-		}
-		
 		// Only include "type": "object" in this top-level schema if it's not a function, and so not following JSON RPC format.
 		if (!bUseJsonRpcFormat)
 		{
@@ -243,7 +217,7 @@ public:
 
 		// If property is to be included - Adds property to JSON and possibly the required list, and also adds to the visitor stack for processing.
 		// Returns whether the property was included.
-		auto MaybeIncludeProperty = [this, PropertyFilter, Flags](const FProperty* Property,
+		auto MaybeIncludeProperty = [this, PropertyFilter, Flags, bUseJsonRpcFormat](const FProperty* Property,
 			const TSharedRef<FJsonObject>& PropertiesSchema, TArray<TSharedPtr<FJsonValue>>& RequiredPropertiesArray) -> bool
 			{
 				const FString PropertyName = Property->GetAuthoredName();
@@ -268,7 +242,7 @@ public:
 				
 				VisitorStack.Emplace(
 					FVisitorTarget(TInPlaceType<FValidConstFProperty>(), Property), 
-					EVisitorStackElementFlags::None, 
+					bUseJsonRpcFormat ? EVisitorStackElementFlags::RequireAllProperties : EVisitorStackElementFlags::None,
 					EditorMetadata.CurrentPropertyMemberPath, 
 					NewPropertyFilter,
 					PropertySchema);
@@ -382,13 +356,18 @@ public:
 	
 private:
 	
-	static void BuildJsonForEnumDefinition(const UEnum& EnumDef, const TSharedRef<FJsonObject>& OutputSchema)
+	static bool BuildJsonForEnumDefinition(const UEnum* EnumDef, const TSharedRef<FJsonObject>& OutputSchema)
 	{
+		if (!EnumDef)
+		{
+			return false;
+		}
+
 		FJsonDomBuilder::FArray EnumNames;
-		const int32 NumEnum = EnumDef.NumEnums();
+		const int32 NumEnum = EnumDef->NumEnums();
 		for (int32 EnumIndex = 0; EnumIndex < NumEnum; ++EnumIndex)
 		{
-			const FString StringValue = EnumDef.GetAuthoredNameStringByIndex(EnumIndex);
+			const FString StringValue = EnumDef->GetAuthoredNameStringByIndex(EnumIndex);
 			if (StringValue.EndsWith(TEXT("_MAX"), ESearchCase::CaseSensitive))
 			{
 				continue;
@@ -398,9 +377,10 @@ private:
 
 		OutputSchema->SetStringField(TEXT("type"), TEXT("string"));
 		OutputSchema->SetField(TEXT("enum"), EnumNames.AsJsonValue());
+		return true;
 	}
 
-	void BuildJsonForOptionalProperty(const FOptionalProperty& OptionalProperty, 
+	bool BuildJsonForOptionalProperty(const FOptionalProperty& OptionalProperty, 
 		const FJsonSchemaPropertyFilter& PropertyFilter, const TSharedRef<FJsonObject>& OutputSchema)
 	{
 		VisitorStack.Emplace(
@@ -409,9 +389,10 @@ private:
 			EditorMetadata.CurrentPropertyMemberPath, 
 			PropertyFilter,
 			OutputSchema);
+		return true;
 	}
 
-	void BuildJsonForStructProperty(const FStructProperty& StructProperty, 
+	bool BuildJsonForStructProperty(const FStructProperty& StructProperty, 
 		const FJsonSchemaPropertyFilter& PropertyFilter, const TSharedRef<FJsonObject>& OutputSchema)
 	{
 		VisitorStack.Emplace(
@@ -420,9 +401,10 @@ private:
 			EditorMetadata.CurrentPropertyMemberPath, 
 			PropertyFilter,
 			OutputSchema);
+		return true;
 	}
 
-	void BuildJsonForArrayProperty(const FArrayProperty& ArrayProperty, 
+	bool BuildJsonForArrayProperty(const FArrayProperty& ArrayProperty, 
 		const FJsonSchemaPropertyFilter& PropertyFilter, const TSharedRef<FJsonObject>& OutputSchema)
 	{
 		const TSharedRef<FJsonObject> ItemsJson = MakeShared<FJsonObject>();
@@ -438,9 +420,10 @@ private:
 		
 		OutputSchema->SetStringField(TEXT("type"), TEXT("array"));
 		OutputSchema->SetObjectField(TEXT("items"), ItemsJson);
+		return true;
 	}
 
-	void BuildJsonForSetProperty(const FSetProperty& SetProperty, 
+	bool BuildJsonForSetProperty(const FSetProperty& SetProperty, 
 		const FJsonSchemaPropertyFilter& PropertyFilter, const TSharedRef<FJsonObject>& OutputSchema)
 	{
 		const TSharedRef<FJsonObject> ItemsJson = MakeShared<FJsonObject>();
@@ -457,14 +440,13 @@ private:
 		OutputSchema->SetStringField(TEXT("type"), TEXT("array"));
 		OutputSchema->SetBoolField(TEXT("uniqueItems"), true);
 		OutputSchema->SetObjectField(TEXT("items"), ItemsJson);
+		return true;
 	}
 
-	void BuildJsonForMapProperty(const FMapProperty& MapProperty, 
+	bool BuildJsonForMapProperty(const FMapProperty& MapProperty, 
 		const FJsonSchemaPropertyFilter& PropertyFilter, const TSharedRef<FJsonObject>& OutputSchema)
 	{
-		if (ensureMsgf(MapProperty.KeyProp->IsA<FStrProperty>(), 
-			TEXT("BuildJsonForMapProperty() - Unsupported map key type for %s. JSON-compatible map keys must be FString. Map will be skipped."), 
-			*EditorMetadata.CurrentPropertyMemberPath))
+		if (MapProperty.KeyProp->IsA<FStrProperty>())
 		{
 			const TSharedRef<FJsonObject> MapValuesJson = MakeShared<FJsonObject>();
 			VisitorStack.Emplace(
@@ -479,37 +461,58 @@ private:
 
 			OutputSchema->SetStringField(TEXT("type"), TEXT("object"));
 			OutputSchema->SetObjectField(TEXT("additionalProperties"), MapValuesJson);
+			return true;
 		}
+		return false;
 	}
 
-	static void BuildJsonForNumericProperty(const FNumericProperty& NumericProperty, const TSharedRef<FJsonObject>& OutputSchema)
+	bool BuildJsonForObjectProperty(const FObjectProperty& ObjectProperty, const TSharedRef<FJsonObject>& OutputSchema)
+	{
+		// This mirrors the behavior of JsonObjectConverter which exports instance objects by value instead of by reference.
+		if (ObjectProperty.HasAnyPropertyFlags(CPF_PersistentInstance))
+		{
+			OutputSchema->SetStringField(TEXT("type"), TEXT("object"));
+		}
+		else
+		{
+			OutputSchema->SetStringField(TEXT("type"), TEXT("string"));
+		}
+		return true;
+	}
+
+	static bool BuildJsonForNumericProperty(const FNumericProperty& NumericProperty, const TSharedRef<FJsonObject>& OutputSchema)
 	{
 		if (const UEnum* EnumDef = NumericProperty.GetIntPropertyEnum())
 		{
-			BuildJsonForEnumDefinition(*EnumDef, OutputSchema);
+			return BuildJsonForEnumDefinition(EnumDef, OutputSchema);
 		}
 		else if (NumericProperty.IsFloatingPoint())
 		{
 			OutputSchema->SetStringField(TEXT("type"), TEXT("number"));
+			return true;
 		}
 		else if (NumericProperty.IsInteger())
 		{
 			OutputSchema->SetStringField(TEXT("type"), TEXT("integer"));
+			return true;
 		}
+		return false;
 	}
 
-	static void BuildJsonForBoolProperty(const TSharedRef<FJsonObject>& OutputSchema)
+	static bool BuildJsonForBoolProperty(const TSharedRef<FJsonObject>& OutputSchema)
 	{
 		OutputSchema->SetStringField(TEXT("type"), TEXT("boolean"));
+		return true;
 	}
 
-	static void BuildJsonForStringProperty(const TSharedRef<FJsonObject>& OutputSchema)
+	static bool BuildJsonForStringProperty(const TSharedRef<FJsonObject>& OutputSchema)
 	{
 		OutputSchema->SetStringField(TEXT("type"), TEXT("string"));
+		return true;
 	}
 
 	FVisitorStack VisitorStack;
-	FJsonSchemaEditorMetadata& EditorMetadata; 
+	FJsonSchemaEditorMetadata& EditorMetadata; 	
 };
 
 //
diff --git a/Engine/Source/Runtime/JsonUtilities/Private/JsonSchema/JsonSchemaPropertyFilter.cpp b/Engine/Source/Runtime/JsonUtilities/Private/JsonSchema/JsonSchemaPropertyFilter.cpp
--- a/Engine/Source/Runtime/JsonUtilities/Private/JsonSchema/JsonSchemaPropertyFilter.cpp
+++ b/Engine/Source/Runtime/JsonUtilities/Private/JsonSchema/JsonSchemaPropertyFilter.cpp
@@ -7,13 +7,15 @@ FJsonSchemaPropertyFilter::FJsonSchemaPropertyFilter(
 	const EPropertyFlags InSkipFlags,
 	const TOptional<TSet<FString>>& InRequiredPropertyMemberPaths,
 	const TOptional<TSet<FString>>& InSkipPropertyMemberPaths,
-	const EJsonObjectConversionFlags InConversionFlags)
+	const EJsonObjectConversionFlags InConversionFlags,
+	CustomCallback* InCustomCb)
 	:
 	CheckFlags(InCheckFlags),
 	SkipFlags(InSkipFlags),
 	RequiredPropertyMemberPaths(InRequiredPropertyMemberPaths),
 	SkipPropertyMemberPaths(InSkipPropertyMemberPaths),
-	ConversionFlags(InConversionFlags)
+	ConversionFlags(InConversionFlags),
+	CustomCb(InCustomCb)
 {
 }
 
@@ -24,6 +26,7 @@ FJsonSchemaPropertyFilter::FJsonSchemaPropertyFilter(const FJsonSchemaPropertyFi
 	RequiredPropertyMemberPaths = Other.RequiredPropertyMemberPaths;
 	SkipPropertyMemberPaths = Other.SkipPropertyMemberPaths;
 	ConversionFlags = Other.ConversionFlags;
+	CustomCb = Other.CustomCb;
 }
 
 bool FJsonSchemaPropertyFilter::IsPropertyIgnored(TNotNull<const FProperty*> Property) const
diff --git a/Engine/Source/Runtime/JsonUtilities/Private/JsonSchema/JsonSchemaVisitorStack.h b/Engine/Source/Runtime/JsonUtilities/Private/JsonSchema/JsonSchemaVisitorStack.h
--- a/Engine/Source/Runtime/JsonUtilities/Private/JsonSchema/JsonSchemaVisitorStack.h
+++ b/Engine/Source/Runtime/JsonUtilities/Private/JsonSchema/JsonSchemaVisitorStack.h
@@ -18,8 +18,9 @@ namespace UE::JsonSchema
 	
 	enum class EVisitorStackElementFlags : uint8
 	{
-		None = 0x0,
-		IsRootStruct = 0x1
+		None					= 0,
+		IsRootStruct			= 1 << 0,
+		RequireAllProperties	= 1 << 1
 	};
 	ENUM_CLASS_FLAGS(EVisitorStackElementFlags);
 
diff --git a/Engine/Source/Runtime/JsonUtilities/Public/JsonSchema/JsonSchemaPropertyFilter.h b/Engine/Source/Runtime/JsonUtilities/Public/JsonSchema/JsonSchemaPropertyFilter.h
--- a/Engine/Source/Runtime/JsonUtilities/Public/JsonSchema/JsonSchemaPropertyFilter.h
+++ b/Engine/Source/Runtime/JsonUtilities/Public/JsonSchema/JsonSchemaPropertyFilter.h
@@ -11,12 +11,20 @@
  */
 struct FJsonSchemaPropertyFilter
 {
+	/**
+	 * Optional callback that will be run when exporting a single property to JsonSchema.
+	 * This function should return true if it's handled the property by writing to OutputSchema.	 
+	 * If ParameterDefaultString is set then only defaults should be written into the OutputSchema object.
+	 */
+	using CustomCallback = TDelegate<bool(const FProperty* Property, const FString& ParameterDefaultString, const TSharedRef<FJsonObject>& OutputSchema)>;
+
 	JSONUTILITIES_API explicit FJsonSchemaPropertyFilter(
 		const EPropertyFlags InCheckFlags = CPF_None,
 		const EPropertyFlags InSkipFlags = CPF_None,
 		const TOptional<TSet<FString>>& InRequiredPropertyMemberPaths = TOptional<TSet<FString>>(),
 		const TOptional<TSet<FString>>& InSkipPropertyMemberPaths = TOptional<TSet<FString>>(),
-		const EJsonObjectConversionFlags InConversionFlags = EJsonObjectConversionFlags::None);
+		const EJsonObjectConversionFlags InConversionFlags = EJsonObjectConversionFlags::None,
+		CustomCallback* CustomCb = nullptr);
 	
 	JSONUTILITIES_API FJsonSchemaPropertyFilter(const FJsonSchemaPropertyFilter& Other);
 	
@@ -51,6 +59,11 @@ struct FJsonSchemaPropertyFilter
 	 * same flags e.g: EJsonObjectConversionFlags::SkipStandardizeCase for property name casing.
 	 */
 	EJsonObjectConversionFlags ConversionFlags = EJsonObjectConversionFlags::None;
+
+	/**
+	 * If passed, allows for custom property handling.
+	 */
+	FJsonSchemaPropertyFilter::CustomCallback* CustomCb = nullptr;
 	
 	/**
 	 * Checks if a property should be ignored, to its flags compared to our check/skip flags.
