Added missing support for TOptional in JsonSchemaMetadata and convert Object defaults to use null rather than an empty string.

Also improved schema tests to ensure we don't have unexpeced fields in the output.

#Synced-CL 50390180

diff --git a/Engine/Plugins/Tests/EditorTests/Source/EditorTests/Private/UnrealEd/JsonUtilitiesTests.cpp b/Engine/Plugins/Tests/EditorTests/Source/EditorTests/Private/UnrealEd/JsonUtilitiesTests.cpp
--- a/Engine/Plugins/Tests/EditorTests/Source/EditorTests/Private/UnrealEd/JsonUtilitiesTests.cpp
+++ b/Engine/Plugins/Tests/EditorTests/Source/EditorTests/Private/UnrealEd/JsonUtilitiesTests.cpp
@@ -32,16 +32,16 @@ END_DEFINE_SPEC(FJsonUtilitiesEditorSpec)
 void TestBasics(FAutomationTestBase& Test, const TSharedPtr<FJsonObject>& Schema, const FString& Type, const FString& Description)
 {
 	Test.TestTrue("Valid schema", Schema.IsValid());
+	FString FieldType;
+	Test.TestEqual("Type", Schema->TryGetStringField("type", FieldType), !Type.IsEmpty());
 	if (!Type.IsEmpty())
 	{
-		FString FieldType;
-		Test.TestTrue("Type", Schema->TryGetStringField("type", FieldType));
 		Test.TestEqual("Type", Type, FieldType);
 	}
+	FString FieldDescription;
+	Test.TestEqual("Description", Schema->TryGetStringField("description", FieldDescription), !Description.IsEmpty());
 	if (!Description.IsEmpty())
 	{
-		FString FieldDescription;
-		Test.TestTrue("Description", Schema->TryGetStringField("description", FieldDescription));
 		Test.TestEqual("Description", Description, FieldDescription);
 	}
 }
@@ -50,46 +50,46 @@ void TestFloat(FAutomationTestBase& Test, const TSharedPtr<FJsonObject>& Schema,
 	TOptional<float> Min = TOptional<float>(), TOptional<float> Max = TOptional<float>(), TOptional<float> Default = TOptional<float>())
 {
 	TestBasics(Test, Schema, "number", Description);
+	float FieldMin;
+	Test.TestEqual("Min", Schema->TryGetNumberField("minimum", FieldMin), (bool)Min);
 	if (Min)
-	{
-		float FieldMin;
-		Test.TestTrue("Min", Schema->TryGetNumberField("minimum", FieldMin));
+	{		
 		Test.TestEqual("Min", Min.GetValue(), FieldMin);
 	}
+	float FieldMax;
+	Test.TestEqual("Max", Schema->TryGetNumberField("maximum", FieldMax), (bool)Max);
 	if (Max)
-	{
-		float FieldMax;
-		Test.TestTrue("Max", Schema->TryGetNumberField("maximum", FieldMax));
+	{		
 		Test.TestEqual("Max", Max.GetValue(), FieldMax);
 	}
+	float FieldDefault;
+	Test.TestEqual("Default", Schema->TryGetNumberField("default", FieldDefault), (bool)Default);
 	if (Default)
 	{
-		float FieldDefault;
-		Test.TestTrue("Default", Schema->TryGetNumberField("default", FieldDefault));
 		Test.TestEqual("Default", Default.GetValue(), FieldDefault);
 	}
 }
 
 void TestInteger(FAutomationTestBase& Test, const TSharedPtr<FJsonObject>& Schema, const FString& Description,
-	TOptional<int> Min = TOptional<int>(), TOptional<int> Max = TOptional<int>(), TOptional<int> Default = TOptional<int>())
+	TOptional<int32> Min = TOptional<int32>(), TOptional<int32> Max = TOptional<int32>(), TOptional<int32> Default = TOptional<int32>())
 {
 	TestBasics(Test, Schema, "integer", Description);
+	int FieldMin;
+	Test.TestEqual("Min", Schema->TryGetNumberField("minimum", FieldMin), (bool)Min);
 	if (Min)
 	{
-		int FieldMin;
-		Test.TestTrue("Min", Schema->TryGetNumberField("minimum", FieldMin));
 		Test.TestEqual("Min", Min.GetValue(), FieldMin);
 	}
+	int FieldMax;
+	Test.TestEqual("Max", Schema->TryGetNumberField("maximum", FieldMax), (bool)Max);
 	if (Max)
 	{
-		int FieldMax;
-		Test.TestTrue("Max", Schema->TryGetNumberField("maximum", FieldMax));
 		Test.TestEqual("Max", Max.GetValue(), FieldMax);
 	}
+	int FieldDefault;
+	Test.TestEqual("Default", Schema->TryGetNumberField("default", FieldDefault), (bool)Default);
 	if (Default)
 	{
-		int FieldDefault;
-		Test.TestTrue("Default", Schema->TryGetNumberField("default", FieldDefault));
 		Test.TestEqual("Default", Default.GetValue(), FieldDefault);
 	}
 }
@@ -97,38 +97,49 @@ void TestInteger(FAutomationTestBase& Test, const TSharedPtr<FJsonObject>& Schem
 void TestBoolean(FAutomationTestBase& Test, const TSharedPtr<FJsonObject>& Schema, const FString& Description, TOptional<bool> Default = TOptional<bool>())
 {
 	TestBasics(Test, Schema, "boolean", Description);
+	bool FieldDefault;
+	Test.TestEqual("Default", Schema->TryGetBoolField("default", FieldDefault), (bool)Default);
 	if (Default)
 	{
-		bool FieldDefault;
-		Test.TestTrue("Default", Schema->TryGetBoolField("default", FieldDefault));
 		Test.TestEqual("Default", Default.GetValue(), FieldDefault);
 	}
 }
 
-void TestString(FAutomationTestBase& Test, const TSharedPtr<FJsonObject>& Schema, const FString& Description, TOptional<FString> Default= TOptional<FString>())
+void TestString(FAutomationTestBase& Test, const TSharedPtr<FJsonObject>& Schema, const FString& Description, TOptional<FString> Default = TOptional<FString>())
 {
 	TestBasics(Test, Schema, "string", Description);
+	FString FieldDefault;
+	Test.TestEqual("Default", Schema->TryGetStringField("default", FieldDefault), (bool)Default);
 	if (Default)
 	{
-		FString FieldDefault;
-		Test.TestTrue("Default", Schema->TryGetStringField("default", FieldDefault));
 		Test.TestEqual("Default", Default.GetValue(), FieldDefault);
 	}
 }
 
+void TestObjectRef(FAutomationTestBase& Test, const TSharedPtr<FJsonObject>& Schema, const FString& Description, bool HasNullDefault)
+{
+	TestBasics(Test, Schema, "string", Description);
+	TSharedPtr<FJsonValue> Default = Schema->TryGetField("default");
+	Test.TestEqual("Default", Default.IsValid(), HasNullDefault);
+	if (Default.IsValid())
+	{				
+		Test.TestTrue("Default", Default->IsNull());		
+	}
+}
+
 void TestEnum(FAutomationTestBase& Test, const TSharedPtr<FJsonObject>& Schema, const FString& Description, const TArray<FString>& Enums, TOptional<FString> Default = TOptional<FString>())
 {
 	TestBasics(Test, Schema, "string", Description);
+	TArray<FString> FieldEnums;
+	Test.TestEqual("Enums", Schema->TryGetStringArrayField("enum", FieldEnums), !Enums.IsEmpty());
 	if (!Enums.IsEmpty())
 	{
-		TArray<FString> FieldEnums;
-		Test.TestTrue("Enums", Schema->TryGetStringArrayField("enum", FieldEnums));
 		Test.TestEqual("Enums", Enums, FieldEnums);
-	}	
+	}
+	FString FieldDefault;
+	Test.TestEqual("Default", Schema->TryGetStringField("default", FieldDefault), (bool)Default);
 	if (Default)
 	{
-		FString FieldDefault;
-		Test.TestTrue("Default", Schema->TryGetStringField("default", FieldDefault));
 		Test.TestEqual("Default", Default.GetValue(), FieldDefault);
 	}
 }
@@ -147,26 +158,30 @@ void TestRequired(FAutomationTestBase& Test, const TSharedPtr<FJsonObject>& Sche
 	}
 }
 
-void TestNumericProperties(FAutomationTestBase& Test, const TSharedPtr<FJsonObject>& Schema, const TArray<FString>& PropertyNames, bool FloatValues = true)
+template <typename T>
+void TestNumericProperties(FAutomationTestBase& Test, const TSharedPtr<FJsonObject>& Schema, const TArray<FString>& PropertyNames, const TArray<FString>& PropertyDescriptions,
+	TOptional<T> Min = TOptional<T>(), TOptional<T> Max = TOptional<T>())
 {
 	const TSharedPtr<FJsonObject>* Properties = nullptr;
 	Test.TestTrue("Properties", Schema->TryGetObjectField("properties", Properties));
 	if (Properties)
 	{		
 		Test.TestEqual("Properties", Properties->Get()->Values.Num(), PropertyNames.Num());
-		for (auto& PropertyName : PropertyNames)
+		for (int PropIndex = 0; PropIndex < PropertyNames.Num(); ++PropIndex)
 		{
+			FString PropertyName = PropertyNames[PropIndex];
+			FString PropertyDescription = PropIndex < PropertyDescriptions.Num() ? PropertyDescriptions[PropIndex] : "";
 			const TSharedPtr<FJsonObject>* Property = nullptr;
 			Test.TestTrue("Properties", Properties->Get()->TryGetObjectField(PropertyName, Property));
 			if (Property)
 			{
-				if (FloatValues)
+				if constexpr (TIsIntegral<T>::Value)
 				{
-					TestFloat(Test, *Property, "");
+					TestInteger(Test, *Property, PropertyDescription, Min, Max);
 				}
 				else
 				{
-					TestInteger(Test, *Property, "");
+					TestFloat(Test, *Property, PropertyDescription, Min, Max);
 				}				
 			}
 		}
@@ -177,21 +192,21 @@ void TestVector(FAutomationTestBase& Test, const TSharedPtr<FJsonObject>& Schema
 {
 	TestBasics(Test, Schema, "object", Description);
 	TArray<FString> PropertyNames({ "x", "y", "z" });
-	TestNumericProperties(Test, Schema, PropertyNames, true);
+	TestNumericProperties<float>(Test, Schema, PropertyNames, TArray<FString>());
 	TestRequired(Test, Schema, PropertyNames);
-		
+	
+	const TSharedPtr<FJsonObject>* FieldDefault = nullptr;
+	Test.TestEqual("Default", Schema->TryGetObjectField("default", FieldDefault), (bool)Default);
 	if(Default)
 	{
-		const TSharedPtr<FJsonObject>* FieldDefault = nullptr;
-		Test.TestTrue("Default", Schema->TryGetObjectField("default", FieldDefault));
 		if (FieldDefault)
 		{
 			Test.TestEqual("Default", FieldDefault->Get()->Values.Num(), PropertyNames.Num());
-			for (int i = 0; i < PropertyNames.Num(); ++i)
+			for (int PropIndex = 0; PropIndex < PropertyNames.Num(); ++PropIndex)
 			{
 				float PropertyValue;
-				Test.TestTrue("Default", FieldDefault->Get()->TryGetNumberField(PropertyNames[i], PropertyValue));
-				Test.TestEqual("Default", PropertyValue, (float)Default.GetValue()[i]);
+				Test.TestTrue("Default", FieldDefault->Get()->TryGetNumberField(PropertyNames[PropIndex], PropertyValue));
+				Test.TestEqual("Default", PropertyValue, (float)Default.GetValue()[PropIndex]);
 			}
 		}
 	}
@@ -201,13 +216,13 @@ void TestQuat(FAutomationTestBase& Test, const TSharedPtr<FJsonObject>& Schema,
 {
 	TestBasics(Test, Schema, "object", Description);
 	TArray<FString> PropertyNames({ "x", "y", "z", "w"});
-	TestNumericProperties(Test, Schema, PropertyNames, true);
+	TestNumericProperties<float>(Test, Schema, PropertyNames, TArray<FString>());
 	TestRequired(Test, Schema, PropertyNames);
 
+	const TSharedPtr<FJsonObject>* FieldDefault = nullptr;
+	Test.TestEqual("Default", Schema->TryGetObjectField("default", FieldDefault), (bool)Default);
 	if (Default)
 	{
-		const TSharedPtr<FJsonObject>* FieldDefault = nullptr;
-		Test.TestTrue("Default", Schema->TryGetObjectField("default", FieldDefault));
 		if (FieldDefault)
 		{
 			Test.TestEqual("Default", FieldDefault->Get()->Values.Num(), PropertyNames.Num());
@@ -240,13 +255,19 @@ void TestRotator(FAutomationTestBase& Test, const TSharedPtr<FJsonObject>& Schem
 {
 	TestBasics(Test, Schema, "object", Description);
 	TArray<FString> PropertyNames({ "yaw", "pitch", "roll" });
-	TestNumericProperties(Test, Schema, PropertyNames, true);
+	TArray<FString> PropertyDescriptions(
+		{
+			"Yaw (degrees) around Z axis",
+			"Pitch (degrees) around Y axis",
+			"Roll (degrees) around X axis"
+		});
+	TestNumericProperties<float>(Test, Schema, PropertyNames, PropertyDescriptions);
 	TestRequired(Test, Schema, PropertyNames);
 
+	const TSharedPtr<FJsonObject>* FieldDefault = nullptr;
+	Test.TestEqual("Default", Schema->TryGetObjectField("default", FieldDefault), (bool)Default);
 	if (Default)
 	{
-		const TSharedPtr<FJsonObject>* FieldDefault = nullptr;
-		Test.TestTrue("Default", Schema->TryGetObjectField("default", FieldDefault));
 		if (FieldDefault)
 		{
 			Test.TestEqual("Default", FieldDefault->Get()->Values.Num(), PropertyNames.Num());
@@ -286,21 +307,21 @@ void TestTransform(FAutomationTestBase& Test, const TSharedPtr<FJsonObject>& Sch
 		Test.TestTrue("Input", Properties->Get()->TryGetObjectField("translation", Translation));
 		if (Translation)
 		{
-			TestVector(Test, *Translation, "");
+			TestVector(Test, *Translation, "Translation of this transformation, as a vector.");
 		}
 
 		const TSharedPtr<FJsonObject>* Quat = nullptr;
 		Test.TestTrue("Input", Properties->Get()->TryGetObjectField("rotation", Quat));
 		if (Quat)
 		{
-			TestQuat(Test, *Quat, "");
+			TestQuat(Test, *Quat, "Rotation of this transformation, as a quaternion.");
 		}
 
 		const TSharedPtr<FJsonObject>* Scale = nullptr;
 		Test.TestTrue("Input", Properties->Get()->TryGetObjectField("scale3D", Scale));
 		if (Scale)
 		{
-			TestVector(Test, *Scale, "");
+			TestVector(Test, *Scale, "3D scale (always applied in local space) as a vector.");
 		}
 	}
 	
@@ -366,13 +387,13 @@ void TestLinearColor(FAutomationTestBase& Test, const TSharedPtr<FJsonObject>& S
 {
 	TestBasics(Test, Schema, "object", Description);
 	TArray<FString> PropertyNames({ "r", "g", "b", "a" });
-	TestNumericProperties(Test, Schema, PropertyNames, true);
+	TestNumericProperties<float>(Test, Schema, PropertyNames, TArray<FString>());
 	TestRequired(Test, Schema, PropertyNames);
 
+	const TSharedPtr<FJsonObject>* FieldDefault = nullptr;
+	Test.TestEqual("Default", Schema->TryGetObjectField("default", FieldDefault), (bool)Default);
 	if (Default)
 	{
-		const TSharedPtr<FJsonObject>* FieldDefault = nullptr;
-		Test.TestTrue("Default", Schema->TryGetObjectField("default", FieldDefault));
 		if (FieldDefault)
 		{
 			Test.TestEqual("Default", FieldDefault->Get()->Values.Num(), PropertyNames.Num());
@@ -405,13 +426,13 @@ void TestColor(FAutomationTestBase& Test, const TSharedPtr<FJsonObject>& Schema,
 {
 	TestBasics(Test, Schema, "object", Description);
 	TArray<FString> PropertyNames({ "r", "g", "b", "a" });
-	TestNumericProperties(Test, Schema, PropertyNames, false);
+	TestNumericProperties<int32>(Test, Schema, PropertyNames, TArray<FString>(), 0, 255);
 	TestRequired(Test, Schema, PropertyNames);
 
+	const TSharedPtr<FJsonObject>* FieldDefault = nullptr;
+	Test.TestEqual("Default", Schema->TryGetObjectField("default", FieldDefault), (bool)Default);
 	if (Default)
 	{
-		const TSharedPtr<FJsonObject>* FieldDefault = nullptr;
-		Test.TestTrue("Default", Schema->TryGetObjectField("default", FieldDefault));
 		if (FieldDefault)
 		{
 			Test.TestEqual("Default", FieldDefault->Get()->Values.Num(), PropertyNames.Num());
@@ -543,7 +564,7 @@ void FJsonUtilitiesEditorSpec::Define()
 					UStruct* Struct = FJsonUtilitiesFakeStruct::StaticStruct();
 					FProperty* Prop = Struct->FindPropertyByName("TestObject");
 					const TSharedPtr<FJsonObject> Schema = FJsonSchemaGenerator::FPropertyToJsonSchemaObject(Prop);
-					TestString(*this, Schema, "TestObjectDesc");
+					TestObjectRef(*this, Schema, "TestObjectDesc", false);
 				});
 
 			It(
@@ -700,7 +721,7 @@ void FJsonUtilitiesEditorSpec::Define()
 						TestTrue("Input", InputSchema->Get()->TryGetObjectField("properties", Properties));
 						if (Properties)
 						{
-							TestEqual("Input", Properties->Get()->Values.Num(), 12);
+							TestEqual("Input", Properties->Get()->Values.Num(), 13);
 							TestRequired(*this, *InputSchema, TArray<FString>({ "requiredParam" }));
 
 							{
@@ -712,6 +733,15 @@ void FJsonUtilitiesEditorSpec::Define()
 								}
 							}
 
+							{
+								const TSharedPtr<FJsonObject>* Property = nullptr;
+								TestTrue("Input", Properties->Get()->TryGetObjectField("optionalParam", Property));
+								if (Property)
+								{
+									TestFloat(*this, *Property, "TestOptionalDesc");
+								}
+							}
+
 							{
 								const TSharedPtr<FJsonObject>* Property = nullptr;
 								TestTrue("Input", Properties->Get()->TryGetObjectField("boolParam", Property));
@@ -735,7 +765,7 @@ void FJsonUtilitiesEditorSpec::Define()
 								TestTrue("Input", Properties->Get()->TryGetObjectField("enumParam", Property));
 								if (Property)
 								{
-									TestEnum(*this, *Property, "TestEnumDesc", TArray<FString>(), FString("One"));
+									TestEnum(*this, *Property, "TestEnumDesc", TArray<FString>({ "Zero", "One" }), FString("One"));
 								}
 							}
 
@@ -762,7 +792,7 @@ void FJsonUtilitiesEditorSpec::Define()
 								TestTrue("Input", Properties->Get()->TryGetObjectField("objectParam", Property));
 								if (Property)
 								{
-									TestString(*this, *Property, "TestObjectDesc", FString(""));
+									TestObjectRef(*this, *Property, "TestObjectDesc", true);
 								}
 							}
 
diff --git a/Engine/Plugins/Tests/EditorTests/Source/EditorTests/Private/UnrealEd/JsonUtilitiesTests.h b/Engine/Plugins/Tests/EditorTests/Source/EditorTests/Private/UnrealEd/JsonUtilitiesTests.h
--- a/Engine/Plugins/Tests/EditorTests/Source/EditorTests/Private/UnrealEd/JsonUtilitiesTests.h
+++ b/Engine/Plugins/Tests/EditorTests/Source/EditorTests/Private/UnrealEd/JsonUtilitiesTests.h
@@ -120,6 +120,7 @@ class UJsonUtilitiesFakeClass : public UObject
 	/**
 	 * TestFuncDesc	 
 	 * @param RequiredParam TestRequiredDesc
+	 * @param OptionalParam TestOptionalDesc
 	 * @param BoolParam TestBoolDesc
 	 * @param FloatParam TestFloatDesc
 	 * @param IntParam TestIntDesc
@@ -134,7 +135,7 @@ class UJsonUtilitiesFakeClass : public UObject
 	 * @return TestReturnDesc
 	 */
 	UFUNCTION(BlueprintCallable, Category="Main") 
-	float TestFunc(int RequiredParam, bool BoolParam = true, float FloatParam = 1.23f,
+	float TestFunc(int RequiredParam, TOptional<float> OptionalParam, bool BoolParam = true, float FloatParam = 1.23f,
 		int32 IntParam = 123, EJsonUtilitiesFakeEnum EnumParam = EJsonUtilitiesFakeEnum::One,
 		const FString& StringParam = "foo", const FName& NameParam = FName("bar"), const UObject* ObjectParam = nullptr,
 		FLinearColor LinearColorParam = FLinearColor(0.1f, 0.2f, 0.3f, 0.4f), FColor ColorParam = FColor(10, 20, 30, 40),
diff --git a/Engine/Source/Editor/JsonUtilitiesEditor/Private/JsonSchema/JsonSchemaGeneratorEditor.cpp b/Engine/Source/Editor/JsonUtilitiesEditor/Private/JsonSchema/JsonSchemaGeneratorEditor.cpp
--- a/Engine/Source/Editor/JsonUtilitiesEditor/Private/JsonSchema/JsonSchemaGeneratorEditor.cpp
+++ b/Engine/Source/Editor/JsonUtilitiesEditor/Private/JsonSchema/JsonSchemaGeneratorEditor.cpp
@@ -379,26 +379,48 @@ namespace FJsonSchemaMetadataPropertyDefaults
 				{
 					UE_LOG(LogJson, Error, TEXT(
 						"Function \"%s\", param property \"%s\". " 
-						"%s.\n"),
+						"%s.\n"
+						"Default string: \"%s\""),
 						*FunctionDebugName, 
 						*Property->GetName(), 
-						*ParameterDefaultStringErrorPhraseString);
+						*ParameterDefaultStringErrorPhraseString,
+						*ParameterDefaultString);
 				}
 			}
 			else if (Property->IsA<FObjectProperty>())
 			{
-				OutPropertyMetadata->SetStringField(TEXT("default"), ParameterDefaultString == "None" ? "" : ParameterDefaultString);				
+				if (ParameterDefaultString == "None")
+				{
+					OutPropertyMetadata->SetField(TEXT("default"), MakeShared<FJsonValueNull>());
+				}
+				else
+				{
+					OutPropertyMetadata->SetStringField(TEXT("default"), ParameterDefaultString);
+				}
+			}
+			else if (const FOptionalProperty* OptionalProperty = CastField<FOptionalProperty>(Property))
+			{
+				// First see if there is a default based on the property the optional contains.
+				SetFunctionPropertyDefaultInPropertyMetadata(
+					PropertyOwnerFunction, OptionalProperty->GetValueProperty(), PropertyFilter, OutPropertyMetadata);
+				// If that fails, make sure there is a default and that it's empty/none.
+				if (!(OutPropertyMetadata->HasField(TEXT("default"))))
+				{
+					OutPropertyMetadata->SetField(TEXT("default"), MakeShared<FJsonValueNull>());
+				}
 			}
 			else
 			{
 				UE_LOG(LogJson, Error, TEXT(
 					"Function \"%s\", param property \"%s\". " 
 					"%s.\n"
-					"Reason: Property type \"%s\" is not recognized."),
+					"Reason: Property type \"%s\" is not recognized.\n"
+					"Default string: \"%s\""),
 					*FunctionDebugName, 
 					*Property->GetName(), 
 					*ParameterDefaultStringErrorPhraseString,
-					*Property->GetClass()->GetName());
+					*Property->GetClass()->GetName(),
+					*ParameterDefaultString);
 			}
 		}
 	}
