Removing deprecated/experimental from JsonSchemaGenerator and fix some use of acutally deprecated APIs that were hiding under blanket deprecation disables.

#Synced-CL 50581518

diff --git a/Engine/Plugins/Experimental/ToolsetRegistry/Source/ToolsetRegistry/Private/ToolsetLibrary.cpp b/Engine/Plugins/Experimental/ToolsetRegistry/Source/ToolsetRegistry/Private/ToolsetLibrary.cpp
--- a/Engine/Plugins/Experimental/ToolsetRegistry/Source/ToolsetRegistry/Private/ToolsetLibrary.cpp
+++ b/Engine/Plugins/Experimental/ToolsetRegistry/Source/ToolsetRegistry/Private/ToolsetLibrary.cpp
@@ -14,9 +14,7 @@ FString UToolsetLibraryTemp::ListClassProperties(UClass* Class)
 	FJsonSchemaPropertyFilter Filter;
 	Filter.CheckFlags |= CPF_BlueprintVisible;
 	Filter.SkipFlags |= CPF_Deprecated;
-PRAGMA_DISABLE_DEPRECATION_WARNINGS
-	return FJsonSchemaGenerator::UStructToJsonSchemaString(Class, Filter);	
-PRAGMA_ENABLE_DEPRECATION_WARNINGS
+	return FJsonSchemaGenerator::UStructToJsonSchemaString(Class, Filter);
 }
 
 FString UToolsetLibraryTemp::GetObjectProperties(const UObject* Object, const TArray<FName>& PropertyNames)
diff --git a/Engine/Plugins/Experimental/ToolsetRegistry/Source/ToolsetRegistry/Private/ToolsetRegistryToolCall.cpp b/Engine/Plugins/Experimental/ToolsetRegistry/Source/ToolsetRegistry/Private/ToolsetRegistryToolCall.cpp
--- a/Engine/Plugins/Experimental/ToolsetRegistry/Source/ToolsetRegistry/Private/ToolsetRegistryToolCall.cpp
+++ b/Engine/Plugins/Experimental/ToolsetRegistry/Source/ToolsetRegistry/Private/ToolsetRegistryToolCall.cpp
@@ -13,7 +13,6 @@
 #include "UObject/FieldIterator.h"
 #include "JsonSchema/JsonSchemaGenerator.h"
 
-PRAGMA_DISABLE_DEPRECATION_WARNINGS
 
 namespace UE::ToolsetRegistry
 {
@@ -501,4 +500,3 @@ namespace UE::ToolsetRegistry
 	}
 }
 
-PRAGMA_ENABLE_DEPRECATION_WARNINGS
diff --git a/Engine/Plugins/Experimental/ToolsetRegistry/Source/ToolsetRegistry/Private/ToolsetRegistryToolset.cpp b/Engine/Plugins/Experimental/ToolsetRegistry/Source/ToolsetRegistry/Private/ToolsetRegistryToolset.cpp
--- a/Engine/Plugins/Experimental/ToolsetRegistry/Source/ToolsetRegistry/Private/ToolsetRegistryToolset.cpp
+++ b/Engine/Plugins/Experimental/ToolsetRegistry/Source/ToolsetRegistry/Private/ToolsetRegistryToolset.cpp
@@ -108,10 +108,8 @@ namespace UE::ToolsetRegistry
 		for (UFunction* Function : ValidTools)
 		{
 			check(Function);
-
-			PRAGMA_DISABLE_DEPRECATION_WARNINGS
-			TSharedPtr<FJsonObject> ToolEntry = FJsonSchemaGenerator::UStructToJsonSchemaObject(Function);
-			PRAGMA_ENABLE_DEPRECATION_WARNINGS
+			
+			TSharedPtr<FJsonObject> ToolEntry = FJsonSchemaGenerator::UStructToJsonSchemaObject(Function);			
 
 			// Modify tool name from "ClassName.FunctionName" to "ToolsetName.FunctionName".
 			FString ToolName = ToolEntry->GetStringField(TEXT("name"));
diff --git a/Engine/Plugins/Tests/EditorTests/Source/EditorTests/Private/UnrealEd/JsonUtilitiesTests.cpp b/Engine/Plugins/Tests/EditorTests/Source/EditorTests/Private/UnrealEd/JsonUtilitiesTests.cpp
--- a/Engine/Plugins/Tests/EditorTests/Source/EditorTests/Private/UnrealEd/JsonUtilitiesTests.cpp
+++ b/Engine/Plugins/Tests/EditorTests/Source/EditorTests/Private/UnrealEd/JsonUtilitiesTests.cpp
@@ -18,10 +18,6 @@
 
 #if WITH_DEV_AUTOMATION_TESTS
 
-// JsonSchema functions are currently marker 'deprecated' as they are still 'experimental'.
-// Remove this later once they are no longer deprecated..
-PRAGMA_DISABLE_DEPRECATION_WARNINGS
-
 #define DEBUG_PRINT_JSONS_TO_LOG 0
 
 
@@ -33,13 +29,13 @@ void TestBasics(FAutomationTestBase& Test, const TSharedPtr<FJsonObject>& Schema
 {
 	Test.TestTrue("Valid schema", Schema.IsValid());
 	FString FieldType;
-	Test.TestEqual("Type", Schema->TryGetStringField("type", FieldType), !Type.IsEmpty());
+	Test.TestEqual("Type", Schema->TryGetStringField(TEXT("type"), FieldType), !Type.IsEmpty());
 	if (!Type.IsEmpty())
 	{
 		Test.TestEqual("Type", Type, FieldType);
 	}
 	FString FieldDescription;
-	Test.TestEqual("Description", Schema->TryGetStringField("description", FieldDescription), !Description.IsEmpty());
+	Test.TestEqual("Description", Schema->TryGetStringField(TEXT("description"), FieldDescription), !Description.IsEmpty());
 	if (!Description.IsEmpty())
 	{
 		Test.TestEqual("Description", Description, FieldDescription);
@@ -51,19 +47,19 @@ void TestFloat(FAutomationTestBase& Test, const TSharedPtr<FJsonObject>& Schema,
 {
 	TestBasics(Test, Schema, "number", Description);
 	float FieldMin;
-	Test.TestEqual("Min", Schema->TryGetNumberField("minimum", FieldMin), (bool)Min);
+	Test.TestEqual("Min", Schema->TryGetNumberField(TEXT("minimum"), FieldMin), (bool)Min);
 	if (Min)
 	{		
 		Test.TestEqual("Min", Min.GetValue(), FieldMin);
 	}
 	float FieldMax;
-	Test.TestEqual("Max", Schema->TryGetNumberField("maximum", FieldMax), (bool)Max);
+	Test.TestEqual("Max", Schema->TryGetNumberField(TEXT("maximum"), FieldMax), (bool)Max);
 	if (Max)
 	{		
 		Test.TestEqual("Max", Max.GetValue(), FieldMax);
 	}
 	float FieldDefault;
-	Test.TestEqual("Default", Schema->TryGetNumberField("default", FieldDefault), (bool)Default);
+	Test.TestEqual("Default", Schema->TryGetNumberField(TEXT("default"), FieldDefault), (bool)Default);
 	if (Default)
 	{
 		Test.TestEqual("Default", Default.GetValue(), FieldDefault);
@@ -75,19 +71,19 @@ void TestInteger(FAutomationTestBase& Test, const TSharedPtr<FJsonObject>& Schem
 {
 	TestBasics(Test, Schema, "integer", Description);
 	int FieldMin;
-	Test.TestEqual("Min", Schema->TryGetNumberField("minimum", FieldMin), (bool)Min);
+	Test.TestEqual("Min", Schema->TryGetNumberField(TEXT("minimum"), FieldMin), (bool)Min);
 	if (Min)
 	{
 		Test.TestEqual("Min", Min.GetValue(), FieldMin);
 	}
 	int FieldMax;
-	Test.TestEqual("Max", Schema->TryGetNumberField("maximum", FieldMax), (bool)Max);
+	Test.TestEqual("Max", Schema->TryGetNumberField(TEXT("maximum"), FieldMax), (bool)Max);
 	if (Max)
 	{
 		Test.TestEqual("Max", Max.GetValue(), FieldMax);
 	}
 	int FieldDefault;
-	Test.TestEqual("Default", Schema->TryGetNumberField("default", FieldDefault), (bool)Default);
+	Test.TestEqual("Default", Schema->TryGetNumberField(TEXT("default"), FieldDefault), (bool)Default);
 	if (Default)
 	{
 		Test.TestEqual("Default", Default.GetValue(), FieldDefault);
@@ -98,7 +94,7 @@ void TestBoolean(FAutomationTestBase& Test, const TSharedPtr<FJsonObject>& Schem
 {
 	TestBasics(Test, Schema, "boolean", Description);
 	bool FieldDefault;
-	Test.TestEqual("Default", Schema->TryGetBoolField("default", FieldDefault), (bool)Default);
+	Test.TestEqual("Default", Schema->TryGetBoolField(TEXT("default"), FieldDefault), (bool)Default);
 	if (Default)
 	{
 		Test.TestEqual("Default", Default.GetValue(), FieldDefault);
@@ -109,7 +105,7 @@ void TestString(FAutomationTestBase& Test, const TSharedPtr<FJsonObject>& Schema
 {
 	TestBasics(Test, Schema, "string", Description);
 	FString FieldDefault;
-	Test.TestEqual("Default", Schema->TryGetStringField("default", FieldDefault), (bool)Default);
+	Test.TestEqual("Default", Schema->TryGetStringField(TEXT("default"), FieldDefault), (bool)Default);
 	if (Default)
 	{
 		Test.TestEqual("Default", Default.GetValue(), FieldDefault);
@@ -119,7 +115,7 @@ void TestString(FAutomationTestBase& Test, const TSharedPtr<FJsonObject>& Schema
 void TestObjectRef(FAutomationTestBase& Test, const TSharedPtr<FJsonObject>& Schema, const FString& Description, bool HasNullDefault)
 {
 	TestBasics(Test, Schema, "string", Description);
-	TSharedPtr<FJsonValue> Default = Schema->TryGetField("default");
+	TSharedPtr<FJsonValue> Default = Schema->TryGetField(TEXT("default"));
 	Test.TestEqual("Default", Default.IsValid(), HasNullDefault);
 	if (Default.IsValid())
 	{				
@@ -131,13 +127,13 @@ void TestEnum(FAutomationTestBase& Test, const TSharedPtr<FJsonObject>& Schema,
 {
 	TestBasics(Test, Schema, "string", Description);
 	TArray<FString> FieldEnums;
-	Test.TestEqual("Enums", Schema->TryGetStringArrayField("enum", FieldEnums), !Enums.IsEmpty());
+	Test.TestEqual("Enums", Schema->TryGetStringArrayField(TEXT("enum"), FieldEnums), !Enums.IsEmpty());
 	if (!Enums.IsEmpty())
 	{
 		Test.TestEqual("Enums", Enums, FieldEnums);
 	}
 	FString FieldDefault;
-	Test.TestEqual("Default", Schema->TryGetStringField("default", FieldDefault), (bool)Default);
+	Test.TestEqual("Default", Schema->TryGetStringField(TEXT("default"), FieldDefault), (bool)Default);
 	if (Default)
 	{
 		Test.TestEqual("Default", Default.GetValue(), FieldDefault);
@@ -147,7 +143,7 @@ void TestEnum(FAutomationTestBase& Test, const TSharedPtr<FJsonObject>& Schema,
 void TestRequired(FAutomationTestBase& Test, const TSharedPtr<FJsonObject>& Schema, const TArray<FString>& Required)
 {
 	TArray<FString> RequiredField;	
-	Test.TestEqual("Required", Schema->TryGetStringArrayField("required", RequiredField), !Required.IsEmpty());
+	Test.TestEqual("Required", Schema->TryGetStringArrayField(TEXT("required"), RequiredField), !Required.IsEmpty());
 	if (!Required.IsEmpty())
 	{
 		Test.TestEqual("Required", Required.Num(), RequiredField.Num());
@@ -163,7 +159,7 @@ void TestNumericProperties(FAutomationTestBase& Test, const TSharedPtr<FJsonObje
 	TOptional<T> Min = TOptional<T>(), TOptional<T> Max = TOptional<T>())
 {
 	const TSharedPtr<FJsonObject>* Properties = nullptr;
-	Test.TestTrue("Properties", Schema->TryGetObjectField("properties", Properties));
+	Test.TestTrue("Properties", Schema->TryGetObjectField(TEXT("properties"), Properties));
 	if (Properties)
 	{		
 		Test.TestEqual("Properties", Properties->Get()->Values.Num(), PropertyNames.Num());
@@ -196,7 +192,7 @@ void TestVector(FAutomationTestBase& Test, const TSharedPtr<FJsonObject>& Schema
 	TestRequired(Test, Schema, PropertyNames);
 	
 	const TSharedPtr<FJsonObject>* FieldDefault = nullptr;
-	Test.TestEqual("Default", Schema->TryGetObjectField("default", FieldDefault), (bool)Default);
+	Test.TestEqual("Default", Schema->TryGetObjectField(TEXT("default"), FieldDefault), (bool)Default);
 	if(Default)
 	{
 		if (FieldDefault)
@@ -220,7 +216,7 @@ void TestQuat(FAutomationTestBase& Test, const TSharedPtr<FJsonObject>& Schema,
 	TestRequired(Test, Schema, PropertyNames);
 
 	const TSharedPtr<FJsonObject>* FieldDefault = nullptr;
-	Test.TestEqual("Default", Schema->TryGetObjectField("default", FieldDefault), (bool)Default);
+	Test.TestEqual("Default", Schema->TryGetObjectField(TEXT("default"), FieldDefault), (bool)Default);
 	if (Default)
 	{
 		if (FieldDefault)
@@ -265,7 +261,7 @@ void TestRotator(FAutomationTestBase& Test, const TSharedPtr<FJsonObject>& Schem
 	TestRequired(Test, Schema, PropertyNames);
 
 	const TSharedPtr<FJsonObject>* FieldDefault = nullptr;
-	Test.TestEqual("Default", Schema->TryGetObjectField("default", FieldDefault), (bool)Default);
+	Test.TestEqual("Default", Schema->TryGetObjectField(TEXT("default"), FieldDefault), (bool)Default);
 	if (Default)
 	{
 		if (FieldDefault)
@@ -299,26 +295,26 @@ void TestTransform(FAutomationTestBase& Test, const TSharedPtr<FJsonObject>& Sch
 	TestRequired(Test, Schema, PropertyNames);
 
 	const TSharedPtr<FJsonObject>* Properties = nullptr;
-	Test.TestTrue("Properties", Schema->TryGetObjectField("properties", Properties));
+	Test.TestTrue("Properties", Schema->TryGetObjectField(TEXT("properties"), Properties));
 	if (Properties)
 	{
 		Test.TestEqual("Properties", Properties->Get()->Values.Num(), PropertyNames.Num());
 		const TSharedPtr<FJsonObject>* Translation = nullptr;
-		Test.TestTrue("Input", Properties->Get()->TryGetObjectField("translation", Translation));
+		Test.TestTrue("Input", Properties->Get()->TryGetObjectField(TEXT("translation"), Translation));
 		if (Translation)
 		{
 			TestVector(Test, *Translation, "Translation of this transformation, as a vector.");
 		}
 
 		const TSharedPtr<FJsonObject>* Quat = nullptr;
-		Test.TestTrue("Input", Properties->Get()->TryGetObjectField("rotation", Quat));
+		Test.TestTrue("Input", Properties->Get()->TryGetObjectField(TEXT("rotation"), Quat));
 		if (Quat)
 		{
 			TestQuat(Test, *Quat, "Rotation of this transformation, as a quaternion.");
 		}
 
 		const TSharedPtr<FJsonObject>* Scale = nullptr;
-		Test.TestTrue("Input", Properties->Get()->TryGetObjectField("scale3D", Scale));
+		Test.TestTrue("Input", Properties->Get()->TryGetObjectField(TEXT("scale3D"), Scale));
 		if (Scale)
 		{
 			TestVector(Test, *Scale, "3D scale (always applied in local space) as a vector.");
@@ -328,13 +324,13 @@ void TestTransform(FAutomationTestBase& Test, const TSharedPtr<FJsonObject>& Sch
 	if (Default)
 	{
 		const TSharedPtr<FJsonObject>* FieldDefault = nullptr;
-		Test.TestTrue("Default", Schema->TryGetObjectField("default", FieldDefault));
+		Test.TestTrue("Default", Schema->TryGetObjectField(TEXT("default"), FieldDefault));
 		if (FieldDefault)
 		{
 			Test.TestEqual("Properties", FieldDefault->Get()->Values.Num(), PropertyNames.Num());
 			{
 				const TSharedPtr<FJsonObject>* Translation = nullptr;
-				Test.TestTrue("Default", FieldDefault->Get()->TryGetObjectField("translation", Translation));
+				Test.TestTrue("Default", FieldDefault->Get()->TryGetObjectField(TEXT("translation"), Translation));
 				if (Translation)
 				{
 					UScriptStruct* VectorStruct = TBaseStructure<FVector>::Get();
@@ -350,7 +346,7 @@ void TestTransform(FAutomationTestBase& Test, const TSharedPtr<FJsonObject>& Sch
 			
 			{
 				const TSharedPtr<FJsonObject>* Quat = nullptr;
-				Test.TestTrue("Default", FieldDefault->Get()->TryGetObjectField("rotation", Quat));
+				Test.TestTrue("Default", FieldDefault->Get()->TryGetObjectField(TEXT("rotation"), Quat));
 				if (Quat)
 				{
 					UScriptStruct* QuatStruct = TBaseStructure<FQuat>::Get();
@@ -366,7 +362,7 @@ void TestTransform(FAutomationTestBase& Test, const TSharedPtr<FJsonObject>& Sch
 			
 			{
 				const TSharedPtr<FJsonObject>* Scale = nullptr;
-				Test.TestTrue("Default", FieldDefault->Get()->TryGetObjectField("scale3D", Scale));
+				Test.TestTrue("Default", FieldDefault->Get()->TryGetObjectField(TEXT("scale3D"), Scale));
 				if (Scale)
 				{
 					UScriptStruct* VectorStruct = TBaseStructure<FVector>::Get();
@@ -391,7 +387,7 @@ void TestLinearColor(FAutomationTestBase& Test, const TSharedPtr<FJsonObject>& S
 	TestRequired(Test, Schema, PropertyNames);
 
 	const TSharedPtr<FJsonObject>* FieldDefault = nullptr;
-	Test.TestEqual("Default", Schema->TryGetObjectField("default", FieldDefault), (bool)Default);
+	Test.TestEqual("Default", Schema->TryGetObjectField(TEXT("default"), FieldDefault), (bool)Default);
 	if (Default)
 	{
 		if (FieldDefault)
@@ -430,7 +426,7 @@ void TestColor(FAutomationTestBase& Test, const TSharedPtr<FJsonObject>& Schema,
 	TestRequired(Test, Schema, PropertyNames);
 
 	const TSharedPtr<FJsonObject>* FieldDefault = nullptr;
-	Test.TestEqual("Default", Schema->TryGetObjectField("default", FieldDefault), (bool)Default);
+	Test.TestEqual("Default", Schema->TryGetObjectField(TEXT("default"), FieldDefault), (bool)Default);
 	if (Default)
 	{
 		if (FieldDefault)
@@ -504,7 +500,7 @@ void FJsonUtilitiesEditorSpec::Define()
 					UStruct* Struct = FJsonUtilitiesFakeStruct::StaticStruct();
 					FProperty* Prop = Struct->FindPropertyByName("TestFloat");
 					const TSharedPtr<FJsonObject> Schema = FJsonSchemaGenerator::FPropertyToJsonSchemaObject(Prop);
-					TestFloat(*this, Schema, "TestFloatDesc", -101.f, 102.f);
+					TestFloat(*this, Schema, "TestFloatDesc", -101.f, 102.f, 1.f);
 				});
 
 			It(
@@ -586,11 +582,11 @@ void FJsonUtilitiesEditorSpec::Define()
 					const TSharedPtr<FJsonObject> Schema = FJsonSchemaGenerator::FPropertyToJsonSchemaObject(Prop);
 					TestBasics(*this, Schema, "object", "TestNestedStructDesc");
 					const TSharedPtr<FJsonObject>* Properties = nullptr;
-					TestTrue("Properties", Schema->TryGetObjectField("properties", Properties));
+					TestTrue("Properties", Schema->TryGetObjectField(TEXT("properties"), Properties));
 					if (Properties)
 					{
 						const TSharedPtr<FJsonObject>* Property = nullptr;
-						TestTrue("Properties", Properties->Get()->TryGetObjectField("TestNestedStructFloat", Property));
+						TestTrue("Properties", Properties->Get()->TryGetObjectField(TEXT("TestNestedStructFloat"), Property));
 						if (Property)
 						{
 							TestFloat(*this, *Property, "TestNestedFloatDesc", -301.f, 302.f);
@@ -607,12 +603,12 @@ void FJsonUtilitiesEditorSpec::Define()
 					const TSharedPtr<FJsonObject> Schema = FJsonSchemaGenerator::UStructToJsonSchemaObject(Struct);
 					TestBasics(*this, Schema, "object", "TestStructDesc");
 					const TSharedPtr<FJsonObject>* Properties = nullptr;
-					TestTrue("Properties", Schema->TryGetObjectField("properties", Properties));
+					TestTrue("Properties", Schema->TryGetObjectField(TEXT("properties"), Properties));
 					if (Properties)
 					{
 						{
 							const TSharedPtr<FJsonObject>* Property = nullptr;
-							TestTrue("Properties", Properties->Get()->TryGetObjectField("TestFloat", Property));
+							TestTrue("Properties", Properties->Get()->TryGetObjectField(TEXT("TestFloat"), Property));
 							if (Property)
 							{
 								TestFloat(*this, *Property, "TestFloatDesc", -101.f, 102.f);
@@ -621,7 +617,7 @@ void FJsonUtilitiesEditorSpec::Define()
 
 						{
 							const TSharedPtr<FJsonObject>* Property = nullptr;
-							TestTrue("Properties", Properties->Get()->TryGetObjectField("TestInt", Property));
+							TestTrue("Properties", Properties->Get()->TryGetObjectField(TEXT("TestInt"), Property));
 							if (Property)
 							{
 								TestInteger(*this, *Property, "TestIntDesc", -201, 202);
@@ -630,7 +626,7 @@ void FJsonUtilitiesEditorSpec::Define()
 
 						{
 							const TSharedPtr<FJsonObject>* Property = nullptr;
-							TestTrue("Properties", Properties->Get()->TryGetObjectField("TestBool", Property));
+							TestTrue("Properties", Properties->Get()->TryGetObjectField(TEXT("TestBool"), Property));
 							if (Property)
 							{
 								TestBoolean(*this, *Property, "TestBoolDesc");
@@ -639,7 +635,7 @@ void FJsonUtilitiesEditorSpec::Define()
 
 						{
 							const TSharedPtr<FJsonObject>* Property = nullptr;
-							TestTrue("Properties", Properties->Get()->TryGetObjectField("TestString", Property));
+							TestTrue("Properties", Properties->Get()->TryGetObjectField(TEXT("TestString"), Property));
 							if (Property)
 							{
 								TestString(*this, *Property, "TestStringDesc");
@@ -648,7 +644,7 @@ void FJsonUtilitiesEditorSpec::Define()
 
 						{
 							const TSharedPtr<FJsonObject>* Property = nullptr;
-							TestTrue("Properties", Properties->Get()->TryGetObjectField("TestName", Property));
+							TestTrue("Properties", Properties->Get()->TryGetObjectField(TEXT("TestName"), Property));
 							if (Property)
 							{
 								TestString(*this, *Property, "TestNameDesc");
@@ -657,7 +653,7 @@ void FJsonUtilitiesEditorSpec::Define()
 
 						{
 							const TSharedPtr<FJsonObject>* Property = nullptr;
-							TestTrue("Properties", Properties->Get()->TryGetObjectField("TestObject", Property));
+							TestTrue("Properties", Properties->Get()->TryGetObjectField(TEXT("TestObject"), Property));
 							if (Property)
 							{
 								TestString(*this, *Property, "TestObjectDesc");
@@ -666,7 +662,7 @@ void FJsonUtilitiesEditorSpec::Define()
 
 						{
 							const TSharedPtr<FJsonObject>* Property = nullptr;
-							TestTrue("Properties", Properties->Get()->TryGetObjectField("TestClass", Property));
+							TestTrue("Properties", Properties->Get()->TryGetObjectField(TEXT("TestClass"), Property));
 							if (Property)
 							{
 								TestString(*this, *Property, "TestClassDesc");
@@ -675,7 +671,7 @@ void FJsonUtilitiesEditorSpec::Define()
 
 						{
 							const TSharedPtr<FJsonObject>* Property = nullptr;
-							TestTrue("Properties", Properties->Get()->TryGetObjectField("TestNestedStruct", Property));
+							TestTrue("Properties", Properties->Get()->TryGetObjectField(TEXT("TestNestedStruct"), Property));
 							if (Property)
 							{
 								TestBasics(*this, *Property, "object", "TestNestedStructDesc");
@@ -691,11 +687,11 @@ void FJsonUtilitiesEditorSpec::Define()
 					const TSharedPtr<FJsonObject> Schema = FJsonSchemaGenerator::UStructToJsonSchemaObject<UJsonUtilitiesFakeClass>();
 					TestBasics(*this, Schema, "object", "TestClassDesc");
 					const TSharedPtr<FJsonObject>* Properties = nullptr;
-					TestTrue("Properties", Schema->TryGetObjectField("properties", Properties));
+					TestTrue("Properties", Schema->TryGetObjectField(TEXT("properties"), Properties));
 					if (Properties)
 					{
 						const TSharedPtr<FJsonObject>* Property = nullptr;
-						TestTrue("Properties", Properties->Get()->TryGetObjectField("TestFloat", Property));
+						TestTrue("Properties", Properties->Get()->TryGetObjectField(TEXT("TestFloat"), Property));
 						if (Property)
 						{
 							TestFloat(*this, *Property, "TestFloatDesc", -101.f, 102.f);
@@ -712,13 +708,13 @@ void FJsonUtilitiesEditorSpec::Define()
 					TestTrue("Schema", Schema.IsValid());
 					TestBasics(*this, Schema, "", "TestFuncDesc");
 					const TSharedPtr<FJsonObject>* InputSchema = nullptr;
-					TestTrue("Input", Schema->TryGetObjectField("inputSchema", InputSchema));
+					TestTrue("Input", Schema->TryGetObjectField(TEXT("inputSchema"), InputSchema));
 					if (InputSchema)
 					{
 						TestBasics(*this, *InputSchema, "object", "");
 
 						const TSharedPtr<FJsonObject>* Properties = nullptr;
-						TestTrue("Input", InputSchema->Get()->TryGetObjectField("properties", Properties));
+						TestTrue("Input", InputSchema->Get()->TryGetObjectField(TEXT("properties"), Properties));
 						if (Properties)
 						{
 							TestEqual("Input", Properties->Get()->Values.Num(), 13);
@@ -726,7 +722,7 @@ void FJsonUtilitiesEditorSpec::Define()
 
 							{
 								const TSharedPtr<FJsonObject>* Property = nullptr;
-								TestTrue("Input", Properties->Get()->TryGetObjectField("requiredParam", Property));
+								TestTrue("Input", Properties->Get()->TryGetObjectField(TEXT("requiredParam"), Property));
 								if (Property)
 								{
 									TestInteger(*this, *Property, "TestRequiredDesc");
@@ -735,7 +731,7 @@ void FJsonUtilitiesEditorSpec::Define()
 
 							{
 								const TSharedPtr<FJsonObject>* Property = nullptr;
-								TestTrue("Input", Properties->Get()->TryGetObjectField("optionalParam", Property));
+								TestTrue("Input", Properties->Get()->TryGetObjectField(TEXT("optionalParam"), Property));
 								if (Property)
 								{
 									TestFloat(*this, *Property, "TestOptionalDesc");
@@ -744,7 +740,7 @@ void FJsonUtilitiesEditorSpec::Define()
 
 							{
 								const TSharedPtr<FJsonObject>* Property = nullptr;
-								TestTrue("Input", Properties->Get()->TryGetObjectField("boolParam", Property));
+								TestTrue("Input", Properties->Get()->TryGetObjectField(TEXT("boolParam"), Property));
 								if (Property)
 								{
 									TestBoolean(*this, *Property, "TestBoolDesc", true);
@@ -753,7 +749,7 @@ void FJsonUtilitiesEditorSpec::Define()
 
 							{
 								const TSharedPtr<FJsonObject>* Property = nullptr;
-								TestTrue("Input", Properties->Get()->TryGetObjectField("floatParam", Property));
+								TestTrue("Input", Properties->Get()->TryGetObjectField(TEXT("floatParam"), Property));
 								if (Property)
 								{
 									TestFloat(*this, *Property, "TestFloatDesc", TOptional<float>(), TOptional<float>(), 1.23f);
@@ -762,7 +758,7 @@ void FJsonUtilitiesEditorSpec::Define()
 
 							{
 								const TSharedPtr<FJsonObject>* Property = nullptr;
-								TestTrue("Input", Properties->Get()->TryGetObjectField("enumParam", Property));
+								TestTrue("Input", Properties->Get()->TryGetObjectField(TEXT("enumParam"), Property));
 								if (Property)
 								{
 									TestEnum(*this, *Property, "TestEnumDesc", TArray<FString>({ "Zero", "One" }), FString("One"));
@@ -771,7 +767,7 @@ void FJsonUtilitiesEditorSpec::Define()
 
 							{
 								const TSharedPtr<FJsonObject>* Property = nullptr;
-								TestTrue("Input", Properties->Get()->TryGetObjectField("stringParam", Property));
+								TestTrue("Input", Properties->Get()->TryGetObjectField(TEXT("stringParam"), Property));
 								if (Property)
 								{
 									TestString(*this, *Property, "TestStringDesc", FString("foo"));
@@ -780,7 +776,7 @@ void FJsonUtilitiesEditorSpec::Define()
 
 							{
 								const TSharedPtr<FJsonObject>* Property = nullptr;
-								TestTrue("Input", Properties->Get()->TryGetObjectField("nameParam", Property));
+								TestTrue("Input", Properties->Get()->TryGetObjectField(TEXT("nameParam"), Property));
 								if (Property)
 								{
 									TestString(*this, *Property, "TestNameDesc", FString("bar"));
@@ -789,7 +785,7 @@ void FJsonUtilitiesEditorSpec::Define()
 
 							{
 								const TSharedPtr<FJsonObject>* Property = nullptr;
-								TestTrue("Input", Properties->Get()->TryGetObjectField("objectParam", Property));
+								TestTrue("Input", Properties->Get()->TryGetObjectField(TEXT("objectParam"), Property));
 								if (Property)
 								{
 									TestObjectRef(*this, *Property, "TestObjectDesc", true);
@@ -798,7 +794,7 @@ void FJsonUtilitiesEditorSpec::Define()
 
 							{
 								const TSharedPtr<FJsonObject>* Property = nullptr;
-								TestTrue("Input", Properties->Get()->TryGetObjectField("linearColorParam", Property));
+								TestTrue("Input", Properties->Get()->TryGetObjectField(TEXT("linearColorParam"), Property));
 								if (Property)
 								{
 									TestLinearColor(*this, *Property, "TestLinearColorDesc", FLinearColor(0.1f, 0.2f, 0.3f, 0.4f));
@@ -807,7 +803,7 @@ void FJsonUtilitiesEditorSpec::Define()
 
 							{
 								const TSharedPtr<FJsonObject>* Property = nullptr;
-								TestTrue("Input", Properties->Get()->TryGetObjectField("colorParam", Property));
+								TestTrue("Input", Properties->Get()->TryGetObjectField(TEXT("colorParam"), Property));
 								if (Property)
 								{
 									TestColor(*this, *Property, "TestColorDesc", FColor(10, 20, 30, 40));
@@ -816,7 +812,7 @@ void FJsonUtilitiesEditorSpec::Define()
 
 							{
 								const TSharedPtr<FJsonObject>* Property = nullptr;
-								TestTrue("Input", Properties->Get()->TryGetObjectField("rotatorParam", Property));
+								TestTrue("Input", Properties->Get()->TryGetObjectField(TEXT("rotatorParam"), Property));
 								if (Property)
 								{
 									TestRotator(*this, *Property, "TestRotatorDesc", FRotator(20, 10, 30));
@@ -825,7 +821,7 @@ void FJsonUtilitiesEditorSpec::Define()
 
 							{
 								const TSharedPtr<FJsonObject>* Property = nullptr;
-								TestTrue("Input", Properties->Get()->TryGetObjectField("vectorParam", Property));
+								TestTrue("Input", Properties->Get()->TryGetObjectField(TEXT("vectorParam"), Property));
 								if (Property)
 								{
 									TestVector(*this, *Property, "TestVectorDesc", FVector(1.f, 2.f, 3.f));
@@ -835,7 +831,7 @@ void FJsonUtilitiesEditorSpec::Define()
 					}
 
 					const TSharedPtr<FJsonObject>* OutputSchema = nullptr;
-					TestTrue("Output", Schema->TryGetObjectField("outputSchema", OutputSchema));
+					TestTrue("Output", Schema->TryGetObjectField(TEXT("outputSchema"), OutputSchema));
 					if (OutputSchema)
 					{
 						TestEqual("Output", OutputSchema->Get()->Values.Num(), 3);
@@ -843,12 +839,12 @@ void FJsonUtilitiesEditorSpec::Define()
 						TestRequired(*this, *OutputSchema, TArray<FString>({ "returnValue" }));
 
 						const TSharedPtr<FJsonObject>* Properties = nullptr;
-						TestTrue("Output", OutputSchema->Get()->TryGetObjectField("properties", Properties));
+						TestTrue("Output", OutputSchema->Get()->TryGetObjectField(TEXT("properties"), Properties));
 						if (Properties)
 						{
 							TestEqual("Output", Properties->Get()->Values.Num(), 1);
 							const TSharedPtr<FJsonObject>* ReturnValue = nullptr;
-							TestTrue("Output", Properties->Get()->TryGetObjectField("returnValue", ReturnValue));
+							TestTrue("Output", Properties->Get()->TryGetObjectField(TEXT("returnValue"), ReturnValue));
 							if (ReturnValue)
 							{
 								TestFloat(*this, *ReturnValue, "TestReturnDesc");
@@ -867,17 +863,17 @@ void FJsonUtilitiesEditorSpec::Define()
 					TestBasics(*this, Schema, "", "TestEngineStructFuncDesc");
 
 					const TSharedPtr<FJsonObject>* InputSchema = nullptr;
-					TestTrue("Input", Schema->TryGetObjectField("inputSchema", InputSchema));
+					TestTrue("Input", Schema->TryGetObjectField(TEXT("inputSchema"), InputSchema));
 					if (InputSchema)
 					{
 						TestBasics(*this, *InputSchema, "object", "");
 						TestRequired(*this, *InputSchema, TArray<FString>());
 						const TSharedPtr<FJsonObject>* Properties = nullptr;
-						TestTrue("Input", InputSchema->Get()->TryGetObjectField("properties", Properties));
+						TestTrue("Input", InputSchema->Get()->TryGetObjectField(TEXT("properties"), Properties));
 						if (Properties)
 						{
 							const TSharedPtr<FJsonObject>* Transform = nullptr;
-							TestTrue("Input", Properties->Get()->TryGetObjectField("transform", Transform));
+							TestTrue("Input", Properties->Get()->TryGetObjectField(TEXT("transform"), Transform));
 							if (Transform)
 							{
 								TestTransform(*this, *Transform, "TestEngineStructFuncParamTransformDesc",
@@ -887,30 +883,30 @@ void FJsonUtilitiesEditorSpec::Define()
 					}
 
 					const TSharedPtr<FJsonObject>* OutputSchema = nullptr;
-					TestTrue("Output", Schema->TryGetObjectField("outputSchema", OutputSchema));
+					TestTrue("Output", Schema->TryGetObjectField(TEXT("outputSchema"), OutputSchema));
 					if (OutputSchema)
 					{
 						TestEqual("Output", OutputSchema->Get()->Values.Num(), 3);
 						TestBasics(*this, *OutputSchema, "object", "");
 						TestRequired(*this, *OutputSchema, TArray<FString>({ "returnValue" }));
 						const TSharedPtr<FJsonObject>* Properties = nullptr;
-						TestTrue("Output", OutputSchema->Get()->TryGetObjectField("properties", Properties));
+						TestTrue("Output", OutputSchema->Get()->TryGetObjectField(TEXT("properties"), Properties));
 						if (Properties)
 						{
 							TestEqual("Output", Properties->Get()->Values.Num(), 1);
 							const TSharedPtr<FJsonObject>* ReturnValue = nullptr;
-							TestTrue("Output", Properties->Get()->TryGetObjectField("returnValue", ReturnValue));
+							TestTrue("Output", Properties->Get()->TryGetObjectField(TEXT("returnValue"), ReturnValue));
 							if (ReturnValue)
 							{
 								TestEqual("Output", ReturnValue->Get()->Values.Num(), 4);
 								TestBasics(*this, *ReturnValue, "object", "TestEngineStructFuncReturnValueDesc");
 								TestRequired(*this, *ReturnValue, TArray<FString>({ "testNestedStructFloat" }));
-								TestTrue("Output", ReturnValue->Get()->TryGetObjectField("properties", Properties));
+								TestTrue("Output", ReturnValue->Get()->TryGetObjectField(TEXT("properties"), Properties));
 								if (Properties)
 								{
 									TestEqual("Output", Properties->Get()->Values.Num(), 1);
 									const TSharedPtr<FJsonObject>* FloatProperty = nullptr;
-									TestTrue("Output", Properties->Get()->TryGetObjectField("testNestedStructFloat", FloatProperty));
+									TestTrue("Output", Properties->Get()->TryGetObjectField(TEXT("testNestedStructFloat"), FloatProperty));
 									if (FloatProperty)
 									{
 										TestFloat(*this, *FloatProperty, "TestNestedFloatDesc", -301.f, 302.f);
@@ -951,17 +947,17 @@ void FJsonUtilitiesEditorSpec::Define()
 					TestTrue("Schema", Schema.IsValid());
 					TestBasics(*this, Schema, "", "TestCustomizedFuncDesc");
 					const TSharedPtr<FJsonObject>* InputSchema = nullptr;
-					TestTrue("Input", Schema->TryGetObjectField("inputSchema", InputSchema));
+					TestTrue("Input", Schema->TryGetObjectField(TEXT("inputSchema"), InputSchema));
 					if (InputSchema)
 					{
 						TestBasics(*this, *InputSchema, "object", "");
 						TestRequired(*this, *InputSchema, TArray<FString>());
 						const TSharedPtr<FJsonObject>* Properties = nullptr;
-						TestTrue("Input", InputSchema->Get()->TryGetObjectField("properties", Properties));
+						TestTrue("Input", InputSchema->Get()->TryGetObjectField(TEXT("properties"), Properties));
 						if (Properties)
 						{
 							const TSharedPtr<FJsonObject>* SoftPath = nullptr;
-							TestTrue("Input", Properties->Get()->TryGetObjectField("softPath", SoftPath));
+							TestTrue("Input", Properties->Get()->TryGetObjectField(TEXT("softPath"), SoftPath));
 							if (SoftPath)
 							{
 								TestString(*this, *SoftPath, "TestSoftPathDesc", FString("/Engine/StaticMesh.StaticMesh"));
@@ -969,10 +965,7 @@ void FJsonUtilitiesEditorSpec::Define()
 						}
 					}
 				});
-	});
+		});
 }
 
-// See comment above.
-PRAGMA_ENABLE_DEPRECATION_WARNINGS
-
 #endif
diff --git a/Engine/Source/Editor/BlueprintEditorLibrary/Private/BlueprintEditorLibrary.cpp b/Engine/Source/Editor/BlueprintEditorLibrary/Private/BlueprintEditorLibrary.cpp
--- a/Engine/Source/Editor/BlueprintEditorLibrary/Private/BlueprintEditorLibrary.cpp
+++ b/Engine/Source/Editor/BlueprintEditorLibrary/Private/BlueprintEditorLibrary.cpp
@@ -1280,9 +1280,7 @@ FString UBlueprintEditorLibrary::PinTypeToJsonSchema(const FEdGraphPinType& PinT
 		return TEXT("");
 	}
 
-PRAGMA_DISABLE_DEPRECATION_WARNINGS
 	TSharedPtr<FJsonObject> PinSchema = FJsonSchemaGenerator::FPropertyToJsonSchemaObject(Property);
-PRAGMA_ENABLE_DEPRECATION_WARNINGS	
 	FString OutString;
 	if(PinSchema)
 	{
diff --git a/Engine/Source/Editor/JsonUtilitiesEditor/Private/JsonUtilitiesEditorModularFeature.cpp b/Engine/Source/Editor/JsonUtilitiesEditor/Private/JsonUtilitiesEditorModularFeature.cpp
--- a/Engine/Source/Editor/JsonUtilitiesEditor/Private/JsonUtilitiesEditorModularFeature.cpp
+++ b/Engine/Source/Editor/JsonUtilitiesEditor/Private/JsonUtilitiesEditorModularFeature.cpp
@@ -4,8 +4,6 @@
 
 #include "JsonSchema/JsonSchemaGeneratorEditor.h"
 
-PRAGMA_DISABLE_DEPRECATION_WARNINGS
-
 /*virtual*/ FJsonSchemaEditorMetadata FJsonUtilitiesEditorModularFeature::FPropertyToJsonSchemaMetadata(TNotNull<const FProperty*> Property,
 	const FJsonSchemaPropertyFilter& PropertyFilter) const
 {
@@ -18,4 +16,3 @@ PRAGMA_DISABLE_DEPRECATION_WARNINGS
 	return FJsonSchemaGeneratorEditor::UStructToJsonSchemaMetadata(Struct, PropertyFilter);
 }
 
-PRAGMA_ENABLE_DEPRECATION_WARNINGS
diff --git a/Engine/Source/Editor/JsonUtilitiesEditor/Public/JsonSchema/JsonSchemaGeneratorEditor.h b/Engine/Source/Editor/JsonUtilitiesEditor/Public/JsonSchema/JsonSchemaGeneratorEditor.h
--- a/Engine/Source/Editor/JsonUtilitiesEditor/Public/JsonSchema/JsonSchemaGeneratorEditor.h
+++ b/Engine/Source/Editor/JsonUtilitiesEditor/Public/JsonSchema/JsonSchemaGeneratorEditor.h
@@ -9,7 +9,7 @@
 #include "JsonSchema/JsonSchemaPropertyFilter.h"
 #include "JsonSchema/JsonSchemaEditorMetadata.h"
 
-class [[deprecated("Experimental class - Subject to change")]] FJsonSchemaGeneratorEditor
+class FJsonSchemaGeneratorEditor
 {
 public:
 	
diff --git a/Engine/Source/Runtime/JsonUtilities/Public/JsonSchema/JsonSchemaGenerator.h b/Engine/Source/Runtime/JsonUtilities/Public/JsonSchema/JsonSchemaGenerator.h
--- a/Engine/Source/Runtime/JsonUtilities/Public/JsonSchema/JsonSchemaGenerator.h
+++ b/Engine/Source/Runtime/JsonUtilities/Public/JsonSchema/JsonSchemaGenerator.h
@@ -12,12 +12,9 @@
 #include "JsonSchemaEditorMetadata.h"
 
 /**
- * Used for generating a JSON schema from any UStruct or FProperty.
- * 
- * See all NOTE_JSON_UTILITIES_EXPERIMENTAL_CLASSES.
- * Currently deprecated to display warnings in build, and warning state in IDEs.
+ * Used for generating a JSON schema from any UStruct or FProperty. 
  */
-class [[deprecated("Experimental class - Subject to change")]] FJsonSchemaGenerator
+class FJsonSchemaGenerator
 {
 public:
 	
@@ -36,7 +33,7 @@ public:
 	 *
 	 * Importantly, the derived JSON Schema describes the JSON which would be produced by FJsonObjectConverter::UPropertyToJsonValue for a property 
 	 * value of this type. And conversely, any JSON adhering to this JSON Schema, may then be deserialized to this property,
-	 * using FJsonObjectConverter::JsonValueToUProperty.     
+	 * using FJsonObjectConverter::JsonValueToUProperty.
 	 *
 	 * e.g:
 	 *
@@ -47,8 +44,8 @@ public:
 	 * @param Property The property to build a schema for.
 	 * @param PropertyFilter Determines how to filter collected struct and property data.
 	 * @param CachedEditorMetadata If valid, use this metadata instead of collecting it. Otherwise, it will be automatically collected
-	 *		(if running in Editor context.) This is optional, and used for workflows that have cached metadata from an earlier process. 
-	 * @return Schema for the FProperty, as a JSON object. 
+	 *		(if running in Editor context.) This is optional, and used for workflows that have cached metadata from an earlier process.
+	 * @return Schema for the FProperty, as a JSON object.
 	 */
 	JSONUTILITIES_API TSharedPtr<FJsonObject> static FPropertyToJsonSchemaObject(TNotNull<const FProperty*> Property,
 		const FJsonSchemaPropertyFilter& PropertyFilter = FJsonSchemaPropertyFilter(), 
@@ -62,10 +59,10 @@ public:
 	 * @param Property The property to build a schema for.
 	 * @param PropertyFilter Determines how to filter collected struct and property data.
 	 * @param CachedEditorMetadata If valid, use this metadata instead of collecting it. Otherwise, it will be automatically collected
-	 *		(if running in Editor context.) This is optional, and used for workflows that have cached metadata from an earlier process. 
+	 *		(if running in Editor context.) This is optional, and used for workflows that have cached metadata from an earlier process.
 	 * @param Indent How many additional tabs to add to the json serializer output.
 	 * @param bPrettyPrint Option to use pretty print (e.g., adds line endings), otherwise condensed print.
-	 * @return Schema for the property, as a JSON string. 
+	 * @return Schema for the property, as a JSON string.
 	 */
 	JSONUTILITIES_API static FString FPropertyToJsonSchemaString(TNotNull<const FProperty*> Property,
 		const FJsonSchemaPropertyFilter& PropertyFilter = FJsonSchemaPropertyFilter(), 
@@ -85,7 +82,7 @@ public:
 	 * 
 	 * Importantly, the derived JSON Schema describes the JSON which would be produced by FJsonObjectConverter::StructToJsonObject for an instance
 	 * of the given UStruct. And conversely, any JSON adhering to this JSON Schema, may then be deserialized to an instance of Struct,
-	 * using FJsonObjectConverter::JsonObjectToUStruct.     
+	 * using FJsonObjectConverter::JsonObjectToUStruct.
 	 *
 	 * e.g:
 	 *
@@ -146,8 +143,8 @@ public:
 	 * @param Struct The type to build a schema for. Can be either UScriptStruct, UClass or UFunction.
 	 * @param PropertyFilter Determines how to filter collected struct and property data.
 	 * @param CachedEditorMetadata If valid, use this metadata instead of collecting it. Otherwise, it will be automatically collected
-	 *		(if running in Editor context.) This is optional, and used for workflows that have cached metadata from an earlier process. 
-	 * @return Schema for the UStruct, as a JSON object. 
+	 *		(if running in Editor context.) This is optional, and used for workflows that have cached metadata from an earlier process.
+	 * @return Schema for the UStruct, as a JSON object.
 	 */
 	JSONUTILITIES_API static TSharedPtr<FJsonObject> UStructToJsonSchemaObject(TNotNull<const UStruct*> Struct,
 		const FJsonSchemaPropertyFilter& PropertyFilter = FJsonSchemaPropertyFilter(),
@@ -161,8 +158,8 @@ public:
 	 * 
 	 * @param PropertyFilter Determines how to filter collected struct and property data.
 	 * @param CachedEditorMetadata If valid, use this metadata instead of collecting it. Otherwise, it will be automatically collected
-	 *		(if running in Editor context.) This is optional, and used for workflows that have cached metadata from an earlier process. 
-	 * @return Schema for the UStruct, as a JSON object. 
+	 *		(if running in Editor context.) This is optional, and used for workflows that have cached metadata from an earlier process.
+	 * @return Schema for the UStruct, as a JSON object.
 	 */
 	template<typename InStructType>
 	static TSharedPtr<FJsonObject> UStructToJsonSchemaObject(
@@ -188,10 +185,10 @@ public:
 	 * @param Struct The type to build a schema for. Can be either UScriptStruct, UClass or UFunction.
 	 * @param PropertyFilter Determines how to filter collected struct and property data.
 	 * @param CachedEditorMetadata If valid, use this metadata instead of collecting it. Otherwise, it will be automatically collected
-	 *		(if running in Editor context.) This is optional, and used for workflows that have cached metadata from an earlier process. 
+	 *		(if running in Editor context.) This is optional, and used for workflows that have cached metadata from an earlier process.
 	 * @param Indent How many additional tabs to add to the json serializer output.
 	 * @param bPrettyPrint Option to use pretty print (e.g., adds line endings), otherwise condensed print.
-	 * @return Schema for the UStruct, as a JSON string. 
+	 * @return Schema for the UStruct, as a JSON string.
 	 */
 	JSONUTILITIES_API static FString UStructToJsonSchemaString(TNotNull<const UStruct*> Struct,
 		const FJsonSchemaPropertyFilter& PropertyFilter = FJsonSchemaPropertyFilter(),
@@ -206,10 +203,10 @@ public:
 	 * 
 	 * @param PropertyFilter Determines how to filter collected struct and property data.
 	 * @param CachedEditorMetadata If valid, use this metadata instead of collecting it. Otherwise, it will be automatically collected
-	 *		(if running in Editor context.) This is optional, and used for workflows that have cached metadata from an earlier process. 
+	 *		(if running in Editor context.) This is optional, and used for workflows that have cached metadata from an earlier process.
 	 * @param Indent How many additional tabs to add to the json serializer output.
 	 * @param bPrettyPrint Option to use pretty print (e.g., adds line endings), otherwise condensed print.
-	 * @return Schema for the UStruct, as a JSON string. 
+	 * @return Schema for the UStruct, as a JSON string.
 	 */
 	template<typename InStructType>
 	static FString UStructToJsonSchemaString(
