Add support for title field in JsonSchema for UStruct, UClass*, UObject*, and UEnum.

#Synced-CL 50581518

diff --git a/Engine/Plugins/Tests/EditorTests/Source/EditorTests/Private/UnrealEd/JsonUtilitiesTests.cpp b/Engine/Plugins/Tests/EditorTests/Source/EditorTests/Private/UnrealEd/JsonUtilitiesTests.cpp
--- a/Engine/Plugins/Tests/EditorTests/Source/EditorTests/Private/UnrealEd/JsonUtilitiesTests.cpp
+++ b/Engine/Plugins/Tests/EditorTests/Source/EditorTests/Private/UnrealEd/JsonUtilitiesTests.cpp
@@ -25,7 +25,9 @@ BEGIN_DEFINE_SPEC(FJsonUtilitiesEditorSpec, "System.Promotion.Editor.JsonUtiliti
 	EAutomationTestFlags::EditorContext|EAutomationTestFlags::EngineFilter)
 END_DEFINE_SPEC(FJsonUtilitiesEditorSpec)
 
-void TestBasics(FAutomationTestBase& Test, const TSharedPtr<FJsonObject>& Schema, const FString& Type, const FString& Description)
+void TestBasics(
+	FAutomationTestBase& Test, const TSharedPtr<FJsonObject>& Schema, const FString& Type, const FString& Description,
+	const TOptional<FString>& Title = TOptional<FString>())
 {
 	Test.TestTrue("Valid schema", Schema.IsValid());
 	FString FieldType;
@@ -40,6 +42,12 @@ void TestBasics(FAutomationTestBase& Test, const TSharedPtr<FJsonObject>& Schema
 	{
 		Test.TestEqual("Description", Description, FieldDescription);
 	}
+	FString FieldTitle;
+	Test.TestEqual("Title", Schema->TryGetStringField(TEXT("title"), FieldTitle), (bool)Title);
+	if (Title)
+	{
+		Test.TestEqual("Title", Title.GetValue(), FieldTitle);
+	}
 }
 
 void TestFloat(FAutomationTestBase& Test, const TSharedPtr<FJsonObject>& Schema, const FString& Description,
@@ -114,7 +122,7 @@ void TestString(FAutomationTestBase& Test, const TSharedPtr<FJsonObject>& Schema
 
 void TestObjectRef(FAutomationTestBase& Test, const TSharedPtr<FJsonObject>& Schema, const FString& Description, bool HasNullDefault)
 {
-	TestBasics(Test, Schema, "string", Description);
+	TestBasics(Test, Schema, "string", Description, FString("Object"));
 	TSharedPtr<FJsonValue> Default = Schema->TryGetField(TEXT("default"));
 	Test.TestEqual("Default", Default.IsValid(), HasNullDefault);
 	if (Default.IsValid())
@@ -123,9 +131,22 @@ void TestObjectRef(FAutomationTestBase& Test, const TSharedPtr<FJsonObject>& Sch
 	}
 }
 
-void TestEnum(FAutomationTestBase& Test, const TSharedPtr<FJsonObject>& Schema, const FString& Description, const TArray<FString>& Enums, TOptional<FString> Default = TOptional<FString>())
+void TestClassRef(FAutomationTestBase& Test, const TSharedPtr<FJsonObject>& Schema, const FString& Description, bool HasNullDefault)
 {
-	TestBasics(Test, Schema, "string", Description);
+	TestBasics(Test, Schema, "string", Description, FString("Class"));
+	TSharedPtr<FJsonValue> Default = Schema->TryGetField(TEXT("default"));
+	Test.TestEqual("Default", Default.IsValid(), HasNullDefault);
+	if (Default.IsValid())
+	{
+		Test.TestTrue("Default", Default->IsNull());
+	}
+}
+
+void TestEnum(
+	FAutomationTestBase& Test, const TSharedPtr<FJsonObject>& Schema, const FString& Description, const FString& Title,
+	const TArray<FString>& Enums, TOptional<FString> Default = TOptional<FString>())
+{
+	TestBasics(Test, Schema, "string", Description, Title);
 	TArray<FString> FieldEnums;
 	Test.TestEqual("Enums", Schema->TryGetStringArrayField(TEXT("enum"), FieldEnums), !Enums.IsEmpty());
 	if (!Enums.IsEmpty())
@@ -186,7 +207,7 @@ void TestNumericProperties(FAutomationTestBase& Test, const TSharedPtr<FJsonObje
 
 void TestVector(FAutomationTestBase& Test, const TSharedPtr<FJsonObject>& Schema, const FString& Description, TOptional<FVector> Default = TOptional<FVector>())
 {
-	TestBasics(Test, Schema, "object", Description);
+	TestBasics(Test, Schema, "object", Description, FString("Vector"));
 	TArray<FString> PropertyNames({ "x", "y", "z" });
 	TestNumericProperties<float>(Test, Schema, PropertyNames, TArray<FString>());
 	TestRequired(Test, Schema, PropertyNames);
@@ -210,7 +231,7 @@ void TestVector(FAutomationTestBase& Test, const TSharedPtr<FJsonObject>& Schema
 
 void TestQuat(FAutomationTestBase& Test, const TSharedPtr<FJsonObject>& Schema, const FString& Description, TOptional<FQuat> Default = TOptional<FQuat>())
 {
-	TestBasics(Test, Schema, "object", Description);
+	TestBasics(Test, Schema, "object", Description, FString("Quat"));
 	TArray<FString> PropertyNames({ "x", "y", "z", "w"});
 	TestNumericProperties<float>(Test, Schema, PropertyNames, TArray<FString>());
 	TestRequired(Test, Schema, PropertyNames);
@@ -249,7 +270,7 @@ void TestQuat(FAutomationTestBase& Test, const TSharedPtr<FJsonObject>& Schema,
 
 void TestRotator(FAutomationTestBase& Test, const TSharedPtr<FJsonObject>& Schema, const FString& Description, TOptional<FRotator> Default = TOptional<FRotator>())
 {
-	TestBasics(Test, Schema, "object", Description);
+	TestBasics(Test, Schema, "object", Description, FString("Rotator"));
 	TArray<FString> PropertyNames({ "yaw", "pitch", "roll" });
 	TArray<FString> PropertyDescriptions(
 		{
@@ -290,7 +311,7 @@ void TestRotator(FAutomationTestBase& Test, const TSharedPtr<FJsonObject>& Schem
 
 void TestTransform(FAutomationTestBase& Test, const TSharedPtr<FJsonObject>& Schema, const FString& Description, TOptional<FTransform> Default = TOptional<FTransform>())
 {
-	TestBasics(Test, Schema, "object", Description);
+	TestBasics(Test, Schema, "object", Description, FString("Transform"));
 	TArray<FString> PropertyNames({ "translation", "rotation", "scale3D" });
 	TestRequired(Test, Schema, PropertyNames);
 
@@ -381,7 +402,7 @@ void TestTransform(FAutomationTestBase& Test, const TSharedPtr<FJsonObject>& Sch
 
 void TestLinearColor(FAutomationTestBase& Test, const TSharedPtr<FJsonObject>& Schema, const FString& Description, TOptional<FLinearColor> Default = TOptional<FLinearColor>())
 {
-	TestBasics(Test, Schema, "object", Description);
+	TestBasics(Test, Schema, "object", Description, FString("LinearColor"));
 	TArray<FString> PropertyNames({ "r", "g", "b", "a" });
 	TestNumericProperties<float>(Test, Schema, PropertyNames, TArray<FString>());
 	TestRequired(Test, Schema, PropertyNames);
@@ -420,7 +441,7 @@ void TestLinearColor(FAutomationTestBase& Test, const TSharedPtr<FJsonObject>& S
 
 void TestColor(FAutomationTestBase& Test, const TSharedPtr<FJsonObject>& Schema, const FString& Description, TOptional<FColor> Default = TOptional<FColor>())
 {
-	TestBasics(Test, Schema, "object", Description);
+	TestBasics(Test, Schema, "object", Description, FString("Color"));
 	TArray<FString> PropertyNames({ "r", "g", "b", "a" });
 	TestNumericProperties<int32>(Test, Schema, PropertyNames, TArray<FString>(), 0, 255);
 	TestRequired(Test, Schema, PropertyNames);
@@ -500,7 +521,7 @@ void FJsonUtilitiesEditorSpec::Define()
 					UStruct* Struct = FJsonUtilitiesFakeStruct::StaticStruct();
 					FProperty* Prop = Struct->FindPropertyByName("TestFloat");
 					const TSharedPtr<FJsonObject> Schema = FJsonSchemaGenerator::FPropertyToJsonSchemaObject(Prop);
-					TestFloat(*this, Schema, "TestFloatDesc", -101.f, 102.f, 1.f);
+					TestFloat(*this, Schema, "TestFloatDesc", -101.f, 102.f);
 				});
 
 			It(
@@ -550,7 +571,7 @@ void FJsonUtilitiesEditorSpec::Define()
 					UStruct* Struct = FJsonUtilitiesFakeStruct::StaticStruct();
 					FProperty* Prop = Struct->FindPropertyByName("TestEnum");
 					const TSharedPtr<FJsonObject> Schema = FJsonSchemaGenerator::FPropertyToJsonSchemaObject(Prop);
-					TestEnum(*this, Schema, "TestEnumDesc", TArray<FString>({ "Zero", "One" }));
+					TestEnum(*this, Schema, "TestEnumDesc", "EJsonUtilitiesFakeEnum", TArray<FString>({"Zero", "One"}));
 				});
 
 			It(
@@ -570,7 +591,7 @@ void FJsonUtilitiesEditorSpec::Define()
 					UStruct* Struct = FJsonUtilitiesFakeStruct::StaticStruct();
 					FProperty* Prop = Struct->FindPropertyByName("TestClass");
 					const TSharedPtr<FJsonObject> Schema = FJsonSchemaGenerator::FPropertyToJsonSchemaObject(Prop);
-					TestString(*this, Schema, "TestClassDesc");
+					TestClassRef(*this, Schema, "TestClassDesc", false);
 				});
 
 			It(
@@ -580,7 +601,7 @@ void FJsonUtilitiesEditorSpec::Define()
 					UStruct* Struct = FJsonUtilitiesFakeStruct::StaticStruct();
 					FProperty* Prop = Struct->FindPropertyByName("TestNestedStruct");
 					const TSharedPtr<FJsonObject> Schema = FJsonSchemaGenerator::FPropertyToJsonSchemaObject(Prop);
-					TestBasics(*this, Schema, "object", "TestNestedStructDesc");
+					TestBasics(*this, Schema, "object", "TestNestedStructDesc", FString("JsonUtilitiesFakeNestedStruct"));
 					const TSharedPtr<FJsonObject>* Properties = nullptr;
 					TestTrue("Properties", Schema->TryGetObjectField(TEXT("properties"), Properties));
 					if (Properties)
@@ -601,7 +622,7 @@ void FJsonUtilitiesEditorSpec::Define()
 				{
 					UStruct* Struct = FJsonUtilitiesFakeStruct::StaticStruct();
 					const TSharedPtr<FJsonObject> Schema = FJsonSchemaGenerator::UStructToJsonSchemaObject(Struct);
-					TestBasics(*this, Schema, "object", "TestStructDesc");
+					TestBasics(*this, Schema, "object", "TestStructDesc", FString("JsonUtilitiesFakeStruct"));
 					const TSharedPtr<FJsonObject>* Properties = nullptr;
 					TestTrue("Properties", Schema->TryGetObjectField(TEXT("properties"), Properties));
 					if (Properties)
@@ -656,7 +677,7 @@ void FJsonUtilitiesEditorSpec::Define()
 							TestTrue("Properties", Properties->Get()->TryGetObjectField(TEXT("TestObject"), Property));
 							if (Property)
 							{
-								TestString(*this, *Property, "TestObjectDesc");
+								TestObjectRef(*this, *Property, "TestObjectDesc", false);
 							}
 						}
 
@@ -665,7 +686,7 @@ void FJsonUtilitiesEditorSpec::Define()
 							TestTrue("Properties", Properties->Get()->TryGetObjectField(TEXT("TestClass"), Property));
 							if (Property)
 							{
-								TestString(*this, *Property, "TestClassDesc");
+								TestClassRef(*this, *Property, "TestClassDesc", false);
 							}
 						}
 
@@ -674,7 +695,7 @@ void FJsonUtilitiesEditorSpec::Define()
 							TestTrue("Properties", Properties->Get()->TryGetObjectField(TEXT("TestNestedStruct"), Property));
 							if (Property)
 							{
-								TestBasics(*this, *Property, "object", "TestNestedStructDesc");
+								TestBasics(*this, *Property, "object", "TestNestedStructDesc", FString("JsonUtilitiesFakeNestedStruct"));
 							}
 						}
 					}
@@ -685,7 +706,7 @@ void FJsonUtilitiesEditorSpec::Define()
 				[this]() -> void
 				{
 					const TSharedPtr<FJsonObject> Schema = FJsonSchemaGenerator::UStructToJsonSchemaObject<UJsonUtilitiesFakeClass>();
-					TestBasics(*this, Schema, "object", "TestClassDesc");
+					TestBasics(*this, Schema, "object", "TestClassDesc", FString("JsonUtilitiesFakeClass"));
 					const TSharedPtr<FJsonObject>* Properties = nullptr;
 					TestTrue("Properties", Schema->TryGetObjectField(TEXT("properties"), Properties));
 					if (Properties)
@@ -761,7 +782,9 @@ void FJsonUtilitiesEditorSpec::Define()
 								TestTrue("Input", Properties->Get()->TryGetObjectField(TEXT("enumParam"), Property));
 								if (Property)
 								{
-									TestEnum(*this, *Property, "TestEnumDesc", TArray<FString>({ "Zero", "One" }), FString("One"));
+									TestEnum(
+										*this, *Property, "TestEnumDesc", "EJsonUtilitiesFakeEnum",
+										TArray<FString>({"Zero", "One"}), FString("One"));
 								}
 							}
 
@@ -898,8 +921,10 @@ void FJsonUtilitiesEditorSpec::Define()
 							TestTrue("Output", Properties->Get()->TryGetObjectField(TEXT("returnValue"), ReturnValue));
 							if (ReturnValue)
 							{
-								TestEqual("Output", ReturnValue->Get()->Values.Num(), 4);
-								TestBasics(*this, *ReturnValue, "object", "TestEngineStructFuncReturnValueDesc");
+								TestEqual("Output", ReturnValue->Get()->Values.Num(), 5);
+								TestBasics(
+									*this, *ReturnValue, "object", "TestEngineStructFuncReturnValueDesc",
+									FString("JsonUtilitiesFakeNestedStruct"));
 								TestRequired(*this, *ReturnValue, TArray<FString>({ "testNestedStructFloat" }));
 								TestTrue("Output", ReturnValue->Get()->TryGetObjectField(TEXT("properties"), Properties));
 								if (Properties)
diff --git a/Engine/Source/Runtime/JsonUtilities/Private/JsonSchema/JsonSchemaGenerator.cpp b/Engine/Source/Runtime/JsonUtilities/Private/JsonSchema/JsonSchemaGenerator.cpp
--- a/Engine/Source/Runtime/JsonUtilities/Private/JsonSchema/JsonSchemaGenerator.cpp
+++ b/Engine/Source/Runtime/JsonUtilities/Private/JsonSchema/JsonSchemaGenerator.cpp
@@ -197,6 +197,7 @@ public:
 		if (!bUseJsonRpcFormat)
 		{
 			OutputSchema->SetStringField(TEXT("type"), TEXT("object"));
+			OutputSchema->SetStringField(TEXT("title"), Struct->GetName());
 		}
 
 		// Special case handling for FJsonObjectWrapper, where we know this must be a Json object.   
@@ -376,6 +377,7 @@ private:
 		}
 
 		OutputSchema->SetStringField(TEXT("type"), TEXT("string"));
+		OutputSchema->SetStringField(TEXT("title"), EnumDef->GetName());
 		OutputSchema->SetField(TEXT("enum"), EnumNames.AsJsonValue());
 		return true;
 	}
@@ -401,6 +403,7 @@ private:
 			EditorMetadata.CurrentPropertyMemberPath, 
 			PropertyFilter,
 			OutputSchema);
+		OutputSchema->SetStringField("title", StructProperty.Struct->GetName());
 		return true;
 	}
 
@@ -414,10 +417,10 @@ private:
 			EditorMetadata.CurrentPropertyMemberPath, 
 			PropertyFilter,
 			ItemsJson);
-		
+
 		// Don't double-up the description for the inner items. Especially as the description will be written for the container e.g: "A list of ..."
 		ItemsJson->RemoveField(TEXT("description"));
-		
+
 		OutputSchema->SetStringField(TEXT("type"), TEXT("array"));
 		OutputSchema->SetObjectField(TEXT("items"), ItemsJson);
 		return true;
@@ -477,6 +480,7 @@ private:
 		{
 			OutputSchema->SetStringField(TEXT("type"), TEXT("string"));
 		}
+		OutputSchema->SetStringField(TEXT("title"), ObjectProperty.PropertyClass.GetName());
 		return true;
 	}
 
