Fix bug JsonSchemaGenerator bug where description was duplicated. Also add test coverage for Array, Set, and Map.

#Synced-CL 51216170

diff --git a/Engine/Plugins/Tests/EditorTests/Source/EditorTests/Private/UnrealEd/JsonUtilitiesTests.cpp b/Engine/Plugins/Tests/EditorTests/Source/EditorTests/Private/UnrealEd/JsonUtilitiesTests.cpp
--- a/Engine/Plugins/Tests/EditorTests/Source/EditorTests/Private/UnrealEd/JsonUtilitiesTests.cpp
+++ b/Engine/Plugins/Tests/EditorTests/Source/EditorTests/Private/UnrealEd/JsonUtilitiesTests.cpp
@@ -616,6 +616,57 @@ void FJsonUtilitiesEditorSpec::Define()
 					}
 				});
 
+			It(
+				"Should generate valid JSON schema from an Array property.",
+				[this]() -> void
+				{
+					UStruct* Struct = FJsonUtilitiesFakeStruct::StaticStruct();
+					FProperty* Prop = Struct->FindPropertyByName("TestArray");
+					const TSharedPtr<FJsonObject> Schema = FJsonSchemaGenerator::FPropertyToJsonSchemaObject(Prop);
+					TestBasics(*this, Schema, "array", "TestArrayDesc");
+					const TSharedPtr<FJsonObject>* Items = nullptr;
+					TestTrue("Items", Schema->TryGetObjectField(TEXT("items"), Items));
+					if (Items)
+					{
+						TestFloat(*this, *Items, "");
+					}
+				});
+
+			It(
+				"Should generate valid JSON schema from a Set property.",
+				[this]() -> void
+				{
+					UStruct* Struct = FJsonUtilitiesFakeStruct::StaticStruct();
+					FProperty* Prop = Struct->FindPropertyByName("TestSet");
+					const TSharedPtr<FJsonObject> Schema = FJsonSchemaGenerator::FPropertyToJsonSchemaObject(Prop);
+					TestBasics(*this, Schema, "array", "TestSetDesc");
+					bool UniqueItems = false;
+					TestTrue("Items", Schema->TryGetBoolField(TEXT("uniqueItems"), UniqueItems));
+					TestTrue("Items", UniqueItems);
+					const TSharedPtr<FJsonObject>* Items = nullptr;
+					TestTrue("Items", Schema->TryGetObjectField(TEXT("items"), Items));
+					if (Items)
+					{
+						TestFloat(*this, *Items, "");
+					}
+				});
+
+			It(
+				"Should generate valid JSON schema from a Map property.",
+				[this]() -> void
+				{
+					UStruct* Struct = FJsonUtilitiesFakeStruct::StaticStruct();
+					FProperty* Prop = Struct->FindPropertyByName("TestMap");
+					const TSharedPtr<FJsonObject> Schema = FJsonSchemaGenerator::FPropertyToJsonSchemaObject(Prop);
+					TestBasics(*this, Schema, "object", "TestMapDesc");
+					const TSharedPtr<FJsonObject>* Items = nullptr;
+					TestTrue("Items", Schema->TryGetObjectField(TEXT("additionalProperties"), Items));
+					if (Items)
+					{
+						TestFloat(*this, *Items, "");
+					}
+				});
+
 			It(
 				"Should generate valid JSON schema from structs.",
 				[this]() -> void
@@ -698,6 +749,33 @@ void FJsonUtilitiesEditorSpec::Define()
 								TestBasics(*this, *Property, "object", "TestNestedStructDesc", FString("JsonUtilitiesFakeNestedStruct"));
 							}
 						}
+
+						{
+							const TSharedPtr<FJsonObject>* Property = nullptr;
+							TestTrue("Properties", Properties->Get()->TryGetObjectField(TEXT("TestArray"), Property));
+							if (Property)
+							{
+								TestBasics(*this, *Property, "array", "TestArrayDesc");
+							}
+						}
+
+						{
+							const TSharedPtr<FJsonObject>* Property = nullptr;
+							TestTrue("Properties", Properties->Get()->TryGetObjectField(TEXT("TestSet"), Property));
+							if (Property)
+							{
+								TestBasics(*this, *Property, "array", "TestSetDesc");
+							}
+						}
+
+						{
+							const TSharedPtr<FJsonObject>* Property = nullptr;
+							TestTrue("Properties", Properties->Get()->TryGetObjectField(TEXT("TestMap"), Property));
+							if (Property)
+							{
+								TestBasics(*this, *Property, "object", "TestMapDesc");
+							}
+						}
 					}
 				});
 
diff --git a/Engine/Plugins/Tests/EditorTests/Source/EditorTests/Private/UnrealEd/JsonUtilitiesTests.h b/Engine/Plugins/Tests/EditorTests/Source/EditorTests/Private/UnrealEd/JsonUtilitiesTests.h
--- a/Engine/Plugins/Tests/EditorTests/Source/EditorTests/Private/UnrealEd/JsonUtilitiesTests.h
+++ b/Engine/Plugins/Tests/EditorTests/Source/EditorTests/Private/UnrealEd/JsonUtilitiesTests.h
@@ -98,9 +98,27 @@ struct FJsonUtilitiesFakeStruct
 	
 	/**
 	 * TestNestedStructDesc
-	 */ 
+	 */
 	UPROPERTY(EditAnywhere, Category="Main")
 	FJsonUtilitiesFakeNestedStruct TestNestedStruct;
+
+	/**
+	 * TestArrayDesc
+	 */
+	UPROPERTY(EditAnywhere, Category = "Main")
+	TArray<float> TestArray;
+
+	/**
+	 * TestSetDesc
+	 */
+	UPROPERTY(EditAnywhere, Category = "Main")
+	TSet<float> TestSet;
+
+	/**
+	 * TestMapDesc
+	 */
+	UPROPERTY(EditAnywhere, Category = "Main")
+	TMap<FString, float> TestMap;
 };
 
 /**
diff --git a/Engine/Source/Runtime/JsonUtilities/Private/JsonSchema/JsonSchemaGenerator.cpp b/Engine/Source/Runtime/JsonUtilities/Private/JsonSchema/JsonSchemaGenerator.cpp
--- a/Engine/Source/Runtime/JsonUtilities/Private/JsonSchema/JsonSchemaGenerator.cpp
+++ b/Engine/Source/Runtime/JsonUtilities/Private/JsonSchema/JsonSchemaGenerator.cpp
@@ -181,6 +181,11 @@ public:
 				OutputSchema->SetField(*JsonMemberName, Tuple.Value);
 			}
 		}
+
+		if (static_cast<uint8>(Flags) & static_cast<uint8>(EVisitorStackElementFlags::SkipDescription))
+		{
+			OutputSchema->RemoveField(TEXT("description"));
+		}
 		
 		return true;
 	}
@@ -413,14 +418,11 @@ private:
 		const TSharedRef<FJsonObject> ItemsJson = MakeShared<FJsonObject>();
 		VisitorStack.Emplace(
 			FVisitorTarget(TInPlaceType<FValidConstFProperty>(), ArrayProperty.Inner), 
-			EVisitorStackElementFlags::None,
+			EVisitorStackElementFlags::SkipDescription,
 			EditorMetadata.CurrentPropertyMemberPath, 
 			PropertyFilter,
 			ItemsJson);
 
-		// Don't double-up the description for the inner items. Especially as the description will be written for the container e.g: "A list of ..."
-		ItemsJson->RemoveField(TEXT("description"));
-
 		OutputSchema->SetStringField(TEXT("type"), TEXT("array"));
 		OutputSchema->SetObjectField(TEXT("items"), ItemsJson);
 		return true;
@@ -432,14 +434,11 @@ private:
 		const TSharedRef<FJsonObject> ItemsJson = MakeShared<FJsonObject>();
 		VisitorStack.Emplace(
 			FVisitorTarget(TInPlaceType<FValidConstFProperty>(), SetProperty.ElementProp), 
-			EVisitorStackElementFlags::None, 
+			EVisitorStackElementFlags::SkipDescription,
 			EditorMetadata.CurrentPropertyMemberPath, 
 			PropertyFilter,
 			ItemsJson);
 
-		// Don't double-up the description for the inner items. Especially as the description will be written for the container e.g: "A list of ..."
-		ItemsJson->RemoveField(TEXT("description"));
-
 		OutputSchema->SetStringField(TEXT("type"), TEXT("array"));
 		OutputSchema->SetBoolField(TEXT("uniqueItems"), true);
 		OutputSchema->SetObjectField(TEXT("items"), ItemsJson);
@@ -454,14 +453,11 @@ private:
 			const TSharedRef<FJsonObject> MapValuesJson = MakeShared<FJsonObject>();
 			VisitorStack.Emplace(
 				FVisitorTarget(TInPlaceType<FValidConstFProperty>(), MapProperty.ValueProp), 
-				EVisitorStackElementFlags::None, 
+				EVisitorStackElementFlags::SkipDescription,
 				EditorMetadata.CurrentPropertyMemberPath, 
 				PropertyFilter,
 				MapValuesJson);
 
-			// Don't double-up the description for the map values. Especially as the description will be written for the container e.g: "A map of ..."
-			MapValuesJson->RemoveField(TEXT("description"));
-
 			OutputSchema->SetStringField(TEXT("type"), TEXT("object"));
 			OutputSchema->SetObjectField(TEXT("additionalProperties"), MapValuesJson);
 			return true;
diff --git a/Engine/Source/Runtime/JsonUtilities/Private/JsonSchema/JsonSchemaVisitorStack.h b/Engine/Source/Runtime/JsonUtilities/Private/JsonSchema/JsonSchemaVisitorStack.h
--- a/Engine/Source/Runtime/JsonUtilities/Private/JsonSchema/JsonSchemaVisitorStack.h
+++ b/Engine/Source/Runtime/JsonUtilities/Private/JsonSchema/JsonSchemaVisitorStack.h
@@ -20,7 +20,8 @@ namespace UE::JsonSchema
 	{
 		None					= 0,
 		IsRootStruct			= 1 << 0,
-		RequireAllProperties	= 1 << 1
+		RequireAllProperties	= 1 << 1,
+		SkipDescription			= 1 << 2
 	};
 	ENUM_CLASS_FLAGS(EVisitorStackElementFlags);
 
