Added JsonSchema support for SoftObjectPtr and SoftClassPtr.

#Synced-CL 51216170

diff --git a/Engine/Plugins/Tests/EditorTests/Source/EditorTests/Private/UnrealEd/JsonUtilitiesTests.cpp b/Engine/Plugins/Tests/EditorTests/Source/EditorTests/Private/UnrealEd/JsonUtilitiesTests.cpp
--- a/Engine/Plugins/Tests/EditorTests/Source/EditorTests/Private/UnrealEd/JsonUtilitiesTests.cpp
+++ b/Engine/Plugins/Tests/EditorTests/Source/EditorTests/Private/UnrealEd/JsonUtilitiesTests.cpp
@@ -120,9 +120,9 @@ void TestString(FAutomationTestBase& Test, const TSharedPtr<FJsonObject>& Schema
 	}
 }
 
-void TestObjectRef(FAutomationTestBase& Test, const TSharedPtr<FJsonObject>& Schema, const FString& Description, bool HasNullDefault)
+void TestObjectRef(FAutomationTestBase& Test, const TSharedPtr<FJsonObject>& Schema, const FString& Description, const FString& Title, bool HasNullDefault)
 {
-	TestBasics(Test, Schema, "string", Description, FString("Object"));
+	TestBasics(Test, Schema, "string", Description, Title);
 	TSharedPtr<FJsonValue> Default = Schema->TryGetField(TEXT("default"));
 	Test.TestEqual("Default", Default.IsValid(), HasNullDefault);
 	if (Default.IsValid())
@@ -581,7 +581,17 @@ void FJsonUtilitiesEditorSpec::Define()
 					UStruct* Struct = FJsonUtilitiesFakeStruct::StaticStruct();
 					FProperty* Prop = Struct->FindPropertyByName("TestObject");
 					const TSharedPtr<FJsonObject> Schema = FJsonSchemaGenerator::FPropertyToJsonSchemaObject(Prop);
-					TestObjectRef(*this, Schema, "TestObjectDesc", false);
+					TestObjectRef(*this, Schema, "TestObjectDesc", FString(TEXT("Actor")), false);
+				});
+
+			It(
+				"Should generate valid JSON schema from SoftObject property.",
+				[this]() -> void
+				{
+					UStruct* Struct = FJsonUtilitiesFakeStruct::StaticStruct();
+					FProperty* Prop = Struct->FindPropertyByName("TestSoftObject");
+					const TSharedPtr<FJsonObject> Schema = FJsonSchemaGenerator::FPropertyToJsonSchemaObject(Prop);
+					TestObjectRef(*this, Schema, "TestSoftObjectDesc", FString(TEXT("Actor")), false);
 				});
 
 			It(
@@ -594,6 +604,16 @@ void FJsonUtilitiesEditorSpec::Define()
 					TestClassRef(*this, Schema, "TestClassDesc", false);
 				});
 
+			It(
+				"Should generate valid JSON schema from SoftClass property.",
+				[this]() -> void
+				{
+					UStruct* Struct = FJsonUtilitiesFakeStruct::StaticStruct();
+					FProperty* Prop = Struct->FindPropertyByName("TestSoftClass");
+					const TSharedPtr<FJsonObject> Schema = FJsonSchemaGenerator::FPropertyToJsonSchemaObject(Prop);
+					TestClassRef(*this, Schema, "TestSoftClassDesc", false);
+				});
+
 			It(
 				"Should generate valid JSON schema from Struct property.",
 				[this]() -> void
@@ -728,7 +748,16 @@ void FJsonUtilitiesEditorSpec::Define()
 							TestTrue("Properties", Properties->Get()->TryGetObjectField(TEXT("TestObject"), Property));
 							if (Property)
 							{
-								TestObjectRef(*this, *Property, "TestObjectDesc", false);
+								TestObjectRef(*this, *Property, "TestObjectDesc", FString(TEXT("Actor")), false);
+							}
+						}
+
+						{
+							const TSharedPtr<FJsonObject>* Property = nullptr;
+							TestTrue("Properties", Properties->Get()->TryGetObjectField(TEXT("TestSoftObject"), Property));
+							if (Property)
+							{
+								TestObjectRef(*this, *Property, "TestSoftObjectDesc", FString(TEXT("Actor")), false);
 							}
 						}
 
@@ -741,6 +770,15 @@ void FJsonUtilitiesEditorSpec::Define()
 							}
 						}
 
+						{
+							const TSharedPtr<FJsonObject>* Property = nullptr;
+							TestTrue("Properties", Properties->Get()->TryGetObjectField(TEXT("TestsoftClass"), Property));
+							if (Property)
+							{
+								TestClassRef(*this, *Property, "TestSoftClassDesc", false);
+							}
+						}
+
 						{
 							const TSharedPtr<FJsonObject>* Property = nullptr;
 							TestTrue("Properties", Properties->Get()->TryGetObjectField(TEXT("TestNestedStruct"), Property));
@@ -889,7 +927,7 @@ void FJsonUtilitiesEditorSpec::Define()
 								TestTrue("Input", Properties->Get()->TryGetObjectField(TEXT("objectParam"), Property));
 								if (Property)
 								{
-									TestObjectRef(*this, *Property, "TestObjectDesc", true);
+									TestObjectRef(*this, *Property, "TestObjectDesc", FString(TEXT("Object")), true);
 								}
 							}
 
diff --git a/Engine/Plugins/Tests/EditorTests/Source/EditorTests/Private/UnrealEd/JsonUtilitiesTests.h b/Engine/Plugins/Tests/EditorTests/Source/EditorTests/Private/UnrealEd/JsonUtilitiesTests.h
--- a/Engine/Plugins/Tests/EditorTests/Source/EditorTests/Private/UnrealEd/JsonUtilitiesTests.h
+++ b/Engine/Plugins/Tests/EditorTests/Source/EditorTests/Private/UnrealEd/JsonUtilitiesTests.h
@@ -88,13 +88,25 @@ struct FJsonUtilitiesFakeStruct
 	 * TestObjectDesc
 	 */
 	UPROPERTY(EditAnywhere, Category = "Main")
-	TObjectPtr<UObject> TestObject;
+	TObjectPtr<AActor> TestObject;
+
+	/**
+	 * TestSoftObjectDesc
+	 */
+	UPROPERTY(EditAnywhere, Category = "Main")
+	TSoftObjectPtr<AActor> TestSoftObject;
 
 	/**
 	 * TestClassDesc
 	 */
 	UPROPERTY(EditAnywhere, Category = "Main")
 	TObjectPtr<UClass> TestClass;
+
+	/**
+	 * TestSoftClassDesc
+	 */
+	UPROPERTY(EditAnywhere, Category = "Main")
+	TSoftClassPtr<UClass> TestSoftClass;
 	
 	/**
 	 * TestNestedStructDesc
diff --git a/Engine/Source/Runtime/JsonUtilities/Private/JsonSchema/JsonSchemaGenerator.cpp b/Engine/Source/Runtime/JsonUtilities/Private/JsonSchema/JsonSchemaGenerator.cpp
--- a/Engine/Source/Runtime/JsonUtilities/Private/JsonSchema/JsonSchemaGenerator.cpp
+++ b/Engine/Source/Runtime/JsonUtilities/Private/JsonSchema/JsonSchemaGenerator.cpp
@@ -120,8 +120,14 @@ public:
 		}
 		else if (const FObjectProperty* ObjectProperty = CastField<FObjectProperty>(Property))
 		{
+			// This will catch UObject* and UClass*.
 			ValidProperty = BuildJsonForObjectProperty(*ObjectProperty, OutputSchema);
-		}		
+		}
+		else if (const FSoftObjectProperty* SoftObjectProperty = CastField<FSoftObjectProperty>(Property))
+		{
+			// This will catch SoftObjectPtr and SoftClassPtr.
+			ValidProperty = BuildJsonForObjectProperty(*SoftObjectProperty, OutputSchema);
+		}
 		else if (const FArrayProperty* ArrayProperty = CastField<FArrayProperty>(Property))
 		{
 			ValidProperty = BuildJsonForArrayProperty(*ArrayProperty, PropertyFilter, OutputSchema);
@@ -465,7 +471,7 @@ private:
 		return false;
 	}
 
-	bool BuildJsonForObjectProperty(const FObjectProperty& ObjectProperty, const TSharedRef<FJsonObject>& OutputSchema)
+	bool BuildJsonForObjectProperty(const FObjectPropertyBase& ObjectProperty, const TSharedRef<FJsonObject>& OutputSchema)
 	{
 		// This mirrors the behavior of JsonObjectConverter which exports instance objects by value instead of by reference.
 		if (ObjectProperty.HasAnyPropertyFlags(CPF_PersistentInstance))
