[AI][JsonUtilities] Add _structType for FInstancedStruct round-trip in JsonObjectConverter and JsonSchemaGenerator

FInstancedStruct is polymorphic -- it wraps an arbitrary struct type at runtime. The existing export (51802455) serializes properties but not the type, and import requires pre-initialization. This breaks MCP round-trip where the receiving side creates FInstancedStruct from scratch. Fix: export includes _structType with the UScriptStruct path, import reads it to call InitializeAs, and the schema generator describes _structType so LLMs know they can create FInstancedStruct from scratch.

Changes:
- JsonObjectConverter export: add _structType field with UScriptStruct::GetPathName()
- JsonObjectConverter import: read _structType, call FindObject + InitializeAs when uninitialized
- JsonObjectConverter import: return true for empty FInstancedStruct without _structType
- JsonSchemaGenerator: add _structType property description in FInstancedStruct static fallback schema
- JsonUtilitiesTests: update 3 FInstancedStruct fallback schema tests to verify _structType

#Synced-CL 51880525

#Synced-CL 51880525

#Synced-CL 51880525

#Synced-CL 51880525

#Synced-CL 51880525
#Synced-CL 51880525

diff --git a/Engine/Plugins/Tests/EditorTests/Source/EditorTests/Private/UnrealEd/JsonUtilitiesTests.cpp b/Engine/Plugins/Tests/EditorTests/Source/EditorTests/Private/UnrealEd/JsonUtilitiesTests.cpp
--- a/Engine/Plugins/Tests/EditorTests/Source/EditorTests/Private/UnrealEd/JsonUtilitiesTests.cpp
+++ b/Engine/Plugins/Tests/EditorTests/Source/EditorTests/Private/UnrealEd/JsonUtilitiesTests.cpp
@@ -1078,15 +1078,19 @@ void FJsonUtilitiesEditorSpec::Define()
 			});
 
 			It(
-				"Should generate a fallback FInstancedStruct schema without instance memory.",
+				"Should generate a fallback FInstancedStruct schema with _structType without instance memory.",
 				[this]() -> void
 				{
 					UStruct* WrapperStruct = FJsonUtilitiesInstancedStructWrapper::StaticStruct();
 					FProperty* Prop = WrapperStruct->FindPropertyByName("InstancedStructField");
 					const TSharedPtr<FJsonObject> Schema = FJsonSchemaGenerator::FPropertyToJsonSchemaObject(Prop);
-					TestBasics(*this, Schema, "object", "", FString("FInstancedStruct"));
-					TestFalse("No properties field", Schema->HasField(TEXT("properties")));
-					TestFalse("No required field", Schema->HasField(TEXT("required")));
+					TestBasics(*this, Schema, "object", "Polymorphic struct. Provide _structType with the UScriptStruct path (e.g. /Script/Module.StructName) to create from scratch.", FString("FInstancedStruct"));
+					const TSharedPtr<FJsonObject>* Properties = nullptr;
+					TestTrue("Has properties", Schema->TryGetObjectField(TEXT("properties"), Properties));
+					if (Properties)
+					{
+						TestTrue("Has _structType", Properties->Get()->HasField(TEXT("_structType")));
+					}
 				});
 
 			It(
@@ -1170,15 +1174,19 @@ void FJsonUtilitiesEditorSpec::Define()
 				});
 
 			It(
-				"Should generate a fallback TOptional<FInstancedStruct> schema without instance memory.",
+				"Should generate a fallback TOptional<FInstancedStruct> schema with _structType without instance memory.",
 				[this]() -> void
 				{
 					UStruct* WrapperStruct = FJsonUtilitiesInstancedStructWrapper::StaticStruct();
 					FProperty* Prop = WrapperStruct->FindPropertyByName("OptionalInstancedStructField");
 					const TSharedPtr<FJsonObject> Schema = FJsonSchemaGenerator::FPropertyToJsonSchemaObject(Prop);
-					TestBasics(*this, Schema, "object", "", FString("FInstancedStruct"));
-					TestFalse("No properties field", Schema->HasField(TEXT("properties")));
-					TestFalse("No required field", Schema->HasField(TEXT("required")));
+					TestBasics(*this, Schema, "object", "Polymorphic struct. Provide _structType with the UScriptStruct path (e.g. /Script/Module.StructName) to create from scratch.", FString("FInstancedStruct"));
+					const TSharedPtr<FJsonObject>* Properties = nullptr;
+					TestTrue("Has properties", Schema->TryGetObjectField(TEXT("properties"), Properties));
+					if (Properties)
+					{
+						TestTrue("Has _structType", Properties->Get()->HasField(TEXT("_structType")));
+					}
 				});
 
 			It(
@@ -1214,7 +1222,7 @@ void FJsonUtilitiesEditorSpec::Define()
 				});
 
 			It(
-				"Should generate a fallback TOptional<FInstancedStruct> schema when instance memory is provided but the optional is not set.",
+				"Should generate a fallback TOptional<FInstancedStruct> schema with _structType when instance memory is provided but the optional is not set.",
 				[this]() -> void
 				{
 					UStruct* WrapperStruct = FJsonUtilitiesInstancedStructWrapper::StaticStruct();
@@ -1222,9 +1230,13 @@ void FJsonUtilitiesEditorSpec::Define()
 					TOptional<FInstancedStruct> OptionalField; // not set
 					const TSharedPtr<FJsonObject> Schema = FJsonSchemaGenerator::FPropertyToJsonSchemaObject(
 						Prop, FJsonSchemaPropertyFilter(), nullptr, &OptionalField);
-					TestBasics(*this, Schema, "object", "", FString("FInstancedStruct"));
-					TestFalse("No properties field", Schema->HasField(TEXT("properties")));
-					TestFalse("No required field", Schema->HasField(TEXT("required")));
+					TestBasics(*this, Schema, "object", "Polymorphic struct. Provide _structType with the UScriptStruct path (e.g. /Script/Module.StructName) to create from scratch.", FString("FInstancedStruct"));
+					const TSharedPtr<FJsonObject>* Properties = nullptr;
+					TestTrue("Has properties", Schema->TryGetObjectField(TEXT("properties"), Properties));
+					if (Properties)
+					{
+						TestTrue("Has _structType", Properties->Get()->HasField(TEXT("_structType")));
+					}
 				});
 
 			It(
diff --git a/Engine/Source/Runtime/JsonUtilities/Private/JsonObjectConverter.cpp b/Engine/Source/Runtime/JsonUtilities/Private/JsonObjectConverter.cpp
--- a/Engine/Source/Runtime/JsonUtilities/Private/JsonObjectConverter.cpp
+++ b/Engine/Source/Runtime/JsonUtilities/Private/JsonObjectConverter.cpp
@@ -260,8 +260,7 @@ TSharedPtr<FJsonValue> ConvertScalarFPropertyToJsonValueWithContainer(FProperty*
 				}
 			}
 
-			// FInstancedStruct: serialize the inner struct's fields as a flat object.
-			// Import assumes the struct is already initialized with the correct type.
+			// FInstancedStruct: serialize the inner struct's fields as a flat object with _structType for round-trip.
 			if (StructProperty->Struct == FInstancedStruct::StaticStruct())
 			{
 				const FInstancedStruct* InstancedStruct = static_cast<const FInstancedStruct*>(Value);
@@ -270,6 +269,7 @@ TSharedPtr<FJsonValue> ConvertScalarFPropertyToJsonValueWithContainer(FProperty*
 					return MakeShared<FJsonValueObject>(MakeShared<FJsonObject>());
 				}
 				TSharedRef<FJsonObject> Out = MakeShared<FJsonObject>();
+				Out->SetStringField(TEXT("_structType"), InstancedStruct->GetScriptStruct()->GetPathName());
 				if (UStructToJsonAttributesWithContainer(InstancedStruct->GetScriptStruct(), InstancedStruct->GetMemory(), Container, ExportedObjects, Out, CheckFlags & (~CPF_ParmFlags), SkipFlags, ExportCb, ConversionFlags))
 				{
 					return MakeShared<FJsonValueObject>(Out);
@@ -903,16 +903,43 @@ namespace
 					// This is the struct that will be used for output.
 					const UScriptStruct* ScriptStruct = nullptr;
 
-					// FInstancedStruct: applies the JSON fields to the struct type already stored in the instance.
-					// The struct must already be initialized with the correct type before importing.
+					// FInstancedStruct: applies the JSON fields to the struct type stored in the instance.
+					// If not yet initialized, reads _structType from JSON to create the correct type.
 					if (StructProperty->Struct == FInstancedStruct::StaticStruct())
 					{
 						FInstancedStruct* InstancedStruct = static_cast<FInstancedStruct*>(OutValue);
+						if (!InstancedStruct->IsValid())
+						{
+							FString StructTypePath;
+							if (Obj->TryGetStringField(TEXT("_structType"), StructTypePath))
+							{
+								UScriptStruct* FoundStruct = FindObject<UScriptStruct>(nullptr, *StructTypePath);
+								if (FoundStruct == nullptr)
+								{
+									UE_LOG(LogJson, Error, TEXT("JsonValueToUProperty - _structType '%s' could not be resolved for %s property %s"),
+										*StructTypePath, *StructProperty->Struct->GetAuthoredName(), *Property->GetAuthoredName());
+									if (OutFailReason)
+									{
+										*OutFailReason = FText::Format(
+											LOCTEXT("FailImportInstancedStructType", "_structType '{0}' could not be resolved for {1} property {2}"),
+											FText::FromString(StructTypePath), FText::FromString(StructProperty->Struct->GetAuthoredName()),
+											FText::FromString(Property->GetAuthoredName()));
+									}
+									return false;
+								}
+								InstancedStruct->InitializeAs(FoundStruct);
+							}
+						}
 						if (InstancedStruct->IsValid())
 						{
 							ScriptStruct = InstancedStruct->GetScriptStruct();
 							OutValue = InstancedStruct->GetMutableMemory();
 						}
+						else
+						{
+							// No _structType and not pre-initialized -- empty FInstancedStruct, leave as-is.
+							return true;
+						}
 					}
 					// FInstancedPropertyBag: flat property values; bag must already be initialized with the correct schema.
 					else if (StructProperty->Struct == FInstancedPropertyBag::StaticStruct())
diff --git a/Engine/Source/Runtime/JsonUtilities/Private/JsonSchema/JsonSchemaGenerator.cpp b/Engine/Source/Runtime/JsonUtilities/Private/JsonSchema/JsonSchemaGenerator.cpp
--- a/Engine/Source/Runtime/JsonUtilities/Private/JsonSchema/JsonSchemaGenerator.cpp
+++ b/Engine/Source/Runtime/JsonUtilities/Private/JsonSchema/JsonSchemaGenerator.cpp
@@ -448,6 +448,15 @@ private:
 			// No instance or invalid struct: static fallback.
 			OutputSchema->SetStringField(TEXT("type"), TEXT("object"));
 			OutputSchema->SetStringField(TEXT("title"), TEXT("FInstancedStruct"));
+			OutputSchema->SetStringField(TEXT("description"),
+				TEXT("Polymorphic struct. Provide _structType with the UScriptStruct path (e.g. /Script/Module.StructName) to create from scratch."));
+			TSharedPtr<FJsonObject> StructTypeSchema = MakeShared<FJsonObject>();
+			StructTypeSchema->SetStringField(TEXT("type"), TEXT("string"));
+			StructTypeSchema->SetStringField(TEXT("description"),
+				TEXT("UScriptStruct path name for polymorphic type identification (e.g. /Script/Module.StructName)."));
+			TSharedPtr<FJsonObject> Properties = MakeShared<FJsonObject>();
+			Properties->SetObjectField(TEXT("_structType"), StructTypeSchema);
+			OutputSchema->SetObjectField(TEXT("properties"), Properties);
 			return true;
 		}
 
