Allow exporting TMap with name based keys, not just strings.

#Synced-CL 52087183

diff --git a/Engine/Plugins/Tests/EditorTests/Source/EditorTests/Private/UnrealEd/JsonUtilitiesTests.cpp b/Engine/Plugins/Tests/EditorTests/Source/EditorTests/Private/UnrealEd/JsonUtilitiesTests.cpp
--- a/Engine/Plugins/Tests/EditorTests/Source/EditorTests/Private/UnrealEd/JsonUtilitiesTests.cpp
+++ b/Engine/Plugins/Tests/EditorTests/Source/EditorTests/Private/UnrealEd/JsonUtilitiesTests.cpp
@@ -674,7 +674,7 @@ void FJsonUtilitiesEditorSpec::Define()
 				});
 
 			It(
-				"Should generate valid JSON schema from a Map property.",
+				"Should generate valid JSON schema from a Map property w/ String key.",
 				[this]() -> void
 				{
 					UStruct* Struct = FJsonUtilitiesFakeStruct::StaticStruct();
@@ -689,6 +689,22 @@ void FJsonUtilitiesEditorSpec::Define()
 					}
 				});
 
+			It(
+				"Should generate valid JSON schema from a Map property w/ Name key.",
+				[this]() -> void
+				{
+					UStruct* Struct = FJsonUtilitiesFakeStruct::StaticStruct();
+					FProperty* Prop = Struct->FindPropertyByName("TestNameMap");
+					const TSharedPtr<FJsonObject> Schema = FJsonSchemaGenerator::FPropertyToJsonSchemaObject(Prop);
+					TestBasics(*this, Schema, "object", "TestNameMapDesc");
+					const TSharedPtr<FJsonObject>* Items = nullptr;
+					TestTrue("Items", Schema->TryGetObjectField(TEXT("additionalProperties"), Items));
+					if (Items)
+					{
+						TestFloat(*this, *Items, "");
+					}
+				});
+
 			It(
 				"Should generate valid JSON schema from structs.",
 				[this]() -> void
@@ -816,6 +832,15 @@ void FJsonUtilitiesEditorSpec::Define()
 								TestBasics(*this, *Property, "object", "TestMapDesc");
 							}
 						}
+
+						{
+							const TSharedPtr<FJsonObject>* Property = nullptr;
+							TestTrue("Properties", Properties->Get()->TryGetObjectField(TEXT("TestNameMap"), Property));
+							if (Property)
+							{
+								TestBasics(*this, *Property, "object", "TestNameMapDesc");
+							}
+						}
 					}
 				});
 
diff --git a/Engine/Plugins/Tests/EditorTests/Source/EditorTests/Private/UnrealEd/JsonUtilitiesTests.h b/Engine/Plugins/Tests/EditorTests/Source/EditorTests/Private/UnrealEd/JsonUtilitiesTests.h
--- a/Engine/Plugins/Tests/EditorTests/Source/EditorTests/Private/UnrealEd/JsonUtilitiesTests.h
+++ b/Engine/Plugins/Tests/EditorTests/Source/EditorTests/Private/UnrealEd/JsonUtilitiesTests.h
@@ -133,6 +133,12 @@ struct FJsonUtilitiesFakeStruct
 	 */
 	UPROPERTY(EditAnywhere, Category = "Main")
 	TMap<FString, float> TestMap;
+
+	/**
+	 * TestNameMapDesc
+	 */
+	UPROPERTY(EditAnywhere, Category = "Main")
+	TMap<FName, float> TestNameMap;
 };
 
 /**
diff --git a/Engine/Source/Runtime/JsonUtilities/Private/JsonSchema/JsonSchemaGenerator.cpp b/Engine/Source/Runtime/JsonUtilities/Private/JsonSchema/JsonSchemaGenerator.cpp
--- a/Engine/Source/Runtime/JsonUtilities/Private/JsonSchema/JsonSchemaGenerator.cpp
+++ b/Engine/Source/Runtime/JsonUtilities/Private/JsonSchema/JsonSchemaGenerator.cpp
@@ -541,7 +541,7 @@ private:
 	bool BuildJsonForMapProperty(const FMapProperty& MapProperty, 
 		const FJsonSchemaPropertyFilter& PropertyFilter, const TSharedRef<FJsonObject>& OutputSchema)
 	{
-		if (MapProperty.KeyProp->IsA<FStrProperty>())
+		if (MapProperty.KeyProp->IsA<FStrProperty>() || MapProperty.KeyProp->IsA<FNameProperty>())
 		{
 			const TSharedRef<FJsonObject> MapValuesJson = MakeShared<FJsonObject>();
 			VisitorStack.Emplace(
